Add SuperSearch module

diff --git a/Engine/Source/Developer/SuperSearch/Private/SSuperSearch.cpp b/Engine/Source/Developer/SuperSearch/Private/SSuperSearch.cpp
new file mode 100644
--- /dev/null
+++ b/Engine/Source/Developer/SuperSearch/Private/SSuperSearch.cpp
@@ -0,0 +1,401 @@
+// Copyright 1998-2014 Epic Games, Inc. All Rights Reserved.
+
+#include "SuperSearchPrivatePCH.h"
+#include "SSuperSearch.h"
+#include "SScrollBorder.h"
+
+SSuperSearchBox::SSuperSearchBox()
+	: bIgnoreUIUpdate(false)
+	, SelectedSuggestion(-1)
+{
+}
+
+BEGIN_SLATE_FUNCTION_BUILD_OPTIMIZATION
+void SSuperSearchBox::Construct( const FArguments& InArgs )
+{
+	ChildSlot
+	[
+		SAssignNew( SuggestionBox, SMenuAnchor )
+			.Placement( InArgs._SuggestionListPlacement )
+			[
+				SAssignNew(InputText, SSearchBox)
+					.OnTextCommitted(this, &SSuperSearchBox::OnTextCommitted)
+					.HintText( NSLOCTEXT( "SuperSearchBox", "HelpHint", "Help" ) )
+					.OnTextChanged(this, &SSuperSearchBox::OnTextChanged)
+					.SelectAllTextWhenFocused(false)
+			]
+			.MenuContent
+			(
+				SNew(SBorder)
+				.BorderImage(FEditorStyle::GetBrush("Menu.Background"))
+				.Padding( FMargin(2) )
+				[
+					SNew(SBox)
+					.HeightOverride(250) // avoids flickering, ideally this would be adaptive to the content without flickering
+					[
+						SAssignNew(SuggestionListView, SListView< TSharedPtr<FSearchEntry> >)
+							.ListItemsSource(&Suggestions)
+							.SelectionMode( ESelectionMode::Single )							// Ideally the mouse over would not highlight while keyboard controls the UI
+							.OnGenerateRow(this, &SSuperSearchBox::MakeSuggestionListItemWidget)
+							.OnSelectionChanged(this, &SSuperSearchBox::SuggestionSelectionChanged)
+							.ItemHeight(18)
+					]
+				]
+			)
+	];
+}
+END_SLATE_FUNCTION_BUILD_OPTIMIZATION
+
+void SSuperSearchBox::Tick( const FGeometry& AllottedGeometry, const double InCurrentTime, const float InDeltaTime )
+{
+	SCompoundWidget::Tick(AllottedGeometry, InCurrentTime, InDeltaTime);
+
+	if (!GIntraFrameDebuggingGameThread && !IsEnabled())
+	{
+		SetEnabled(true);
+	}
+	else if (GIntraFrameDebuggingGameThread && IsEnabled())
+	{
+		SetEnabled(false);
+	}
+}
+
+void SSuperSearchBox::ActOnSuggestion(TSharedPtr<FSearchEntry> SearchEntry)
+{
+	if (SearchEntry->bCategory == false)
+	{
+		FPlatformProcess::LaunchURL(*SearchEntry->URL, NULL, NULL);
+	}
+}
+
+void SSuperSearchBox::SuggestionSelectionChanged(TSharedPtr<FSearchEntry> NewValue, ESelectInfo::Type SelectInfo)
+{
+	if(bIgnoreUIUpdate)
+	{
+		return;
+	}
+
+	if (NewValue->bCategory)
+	{
+		return;
+	}
+
+	for(int32 i = 0; i < Suggestions.Num(); ++i)
+	{
+		if(NewValue == Suggestions[i])
+		{
+			SelectedSuggestion = i;
+			MarkActiveSuggestion();
+
+			if( SelectInfo == ESelectInfo::OnMouseClick )
+			{
+				ActOnSuggestion(NewValue);
+			}
+
+			break;
+		}
+	}
+}
+
+
+TSharedRef<ITableRow> SSuperSearchBox::MakeSuggestionListItemWidget(TSharedPtr<FSearchEntry> Suggestion, const TSharedRef<STableViewBase>& OwnerTable)
+{
+	check(Suggestion.IsValid());
+
+	FString Left, Right, Combined;
+
+	if(Suggestion->Title.Split(TEXT("\t"), &Left, &Right))
+	{
+		Combined = Left + Right;
+	}
+	else
+	{
+		Combined = Suggestion->Title;
+	}
+
+	FText HighlightText = FText::FromString(Left);
+
+	bool bCategory = Suggestion->bCategory;
+
+	return
+		SNew(STableRow< TSharedPtr<FString> >, OwnerTable)
+		.ShowSelection(bCategory == false)
+		[
+			SNew(SBox)
+			.WidthOverride(600)			// to enforce some minimum width, ideally we define the minimum, not a fixed width
+			.Padding(FMargin(bCategory ? 0 : 16, 0, 0, 0))
+			[
+				SNew(STextBlock)
+				.Text(Combined)
+				//.TextStyle( FEditorStyle::Get(), "Log.Normal")
+				.HighlightText(HighlightText)
+			]
+		];
+}
+
+void SSuperSearchBox::OnTextChanged(const FText& InText)
+{
+	if(bIgnoreUIUpdate)
+	{
+		return;
+	}
+
+	const FString& InputTextStr = InputText->GetText().ToString();
+	if(!InputTextStr.IsEmpty())
+	{
+		
+		TSharedRef<class IHttpRequest> HttpRequest = FHttpModule::Get().CreateRequest();
+
+		// build the url
+		// use partial response to only look at items and things we care about for items right now (title,link and label name)
+		const FText QueryURL = FText::Format( FText::FromString("https://www.googleapis.com/customsearch/v1?key=AIzaSyCMGfdDaSfjqv5zYoS0mTJnOT3e9MURWkU&cx=009868829633250020713:y7tfd8hlcgg&fields=items(title,link,labels/name)&q={0}+less:forums"), InText);
+
+		//save http request into map to ensure correct ordering
+		FText & Query = RequestQueryMap.FindOrAdd(HttpRequest);
+		Query = InText;
+
+		// kick off http request to read friends
+		HttpRequest->OnProcessRequestComplete().BindRaw(this, &SSuperSearchBox::Query_HttpRequestComplete);
+		HttpRequest->SetURL(QueryURL.ToString());
+		HttpRequest->SetHeader(TEXT("Content-Type"), TEXT("application/json"));
+		HttpRequest->SetVerb(TEXT("GET"));
+		HttpRequest->ProcessRequest();
+	}
+	else
+	{
+		ClearSuggestions();
+	}
+}
+
+void SSuperSearchBox::OnTextCommitted( const FText& InText, ETextCommit::Type CommitInfo)
+{
+	if (CommitInfo == ETextCommit::OnEnter && SelectedSuggestion >= 1)
+	{
+		ActOnSuggestion(Suggestions[SelectedSuggestion]);
+	}
+}
+
+void SSuperSearchBox::Query_HttpRequestComplete(FHttpRequestPtr HttpRequest, FHttpResponsePtr HttpResponse, bool bSucceeded)
+{
+	bool bResult = false;
+	FString ResponseStr, ErrorStr;
+
+	if (bSucceeded && HttpResponse.IsValid())
+	{
+		ResponseStr = HttpResponse->GetContentAsString();
+		if (EHttpResponseCodes::IsOk(HttpResponse->GetResponseCode()))
+		{
+			// Create the Json parser
+			TSharedPtr<FJsonObject> JsonObject;
+			TSharedRef<TJsonReader<> > JsonReader = TJsonReaderFactory<>::Create(ResponseStr);
+
+			if (FJsonSerializer::Deserialize(JsonReader, JsonObject) && JsonObject.IsValid())
+			{
+				if (FText * QueryText = RequestQueryMap.Find(HttpRequest))
+				{
+					const TArray<TSharedPtr<FJsonValue> > * JsonItems = nullptr;
+					if (JsonObject->TryGetArrayField(TEXT("items"), JsonItems))
+					{
+						//add search result into cache
+						FSearchResults & SearchResults = SearchResultsCache.FindOrAdd(QueryText->ToString());
+						FCategoryResults & OnlineResults = SearchResults.OnlineResults;
+						OnlineResults.Empty();	//reset online results since we just got updated
+
+						for (TSharedPtr<FJsonValue> JsonItem : *JsonItems)
+						{
+							const TSharedPtr<FJsonObject> & ItemObject = JsonItem->AsObject();
+							FSearchEntry SearchEntry;
+							SearchEntry.Title = ItemObject->GetStringField("title");
+							SearchEntry.URL = ItemObject->GetStringField("link");
+							SearchEntry.bCategory = false;
+
+							bool bValidLabel = false;
+							TArray<TSharedPtr<FJsonValue> > Labels = ItemObject->GetArrayField(TEXT("labels"));
+							for (TSharedPtr<FJsonValue> Label : Labels)
+							{
+								const TSharedPtr<FJsonObject> & LabelObject = Label->AsObject();
+								FString LabelString = LabelObject->GetStringField(TEXT("name"));
+								TArray<FSearchEntry> & SuggestionCategory = OnlineResults.FindOrAdd(LabelString);
+								SuggestionCategory.Add(SearchEntry);
+							}
+						}
+					}
+
+					bResult = true;
+				}
+				
+			}
+		}
+		else
+		{
+			ErrorStr = FString::Printf(TEXT("Invalid response. code=%d error=%s"),
+				HttpResponse->GetResponseCode(), *ResponseStr);
+		}
+	}
+
+	if (bResult)
+	{
+		UpdateSuggestions();
+	}
+}
+
+
+FReply SSuperSearchBox::OnKeyDown( const FGeometry& MyGeometry, const FKeyboardEvent& KeyboardEvent )
+{
+	if(SuggestionBox->IsOpen() && Suggestions.Num())
+	{
+		if(KeyboardEvent.GetKey() == EKeys::Up || KeyboardEvent.GetKey() == EKeys::Down)
+		{
+			if(KeyboardEvent.GetKey() == EKeys::Up)
+			{
+				//if we're at the top swing around
+				if(SelectedSuggestion == 1)
+				{
+					SelectedSuggestion = Suggestions.Num() - 1;
+				}
+				else
+				{
+					// got one up
+					--SelectedSuggestion;
+
+					//make sure we're not selecting category
+					if (Suggestions[SelectedSuggestion]->bCategory)
+					{
+						//we know that category is never shown empty so above us will be fine
+						--SelectedSuggestion;
+					}
+				}
+			}
+
+			if(KeyboardEvent.GetKey() == EKeys::Down)
+			{
+				if(SelectedSuggestion < Suggestions.Num() - 1)
+				{
+					// go one down, possibly from edit control to top
+					++SelectedSuggestion;
+
+					//make sure we're not selecting category
+					if (Suggestions[SelectedSuggestion]->bCategory)
+					{
+						//we know that category is never shown empty so below us will be fine
+						++SelectedSuggestion;
+					}
+				}
+				else
+				{
+					// back to edit control
+					SelectedSuggestion = 1;
+				}
+			}
+
+			MarkActiveSuggestion();
+
+			return FReply::Handled();
+		}
+	}
+	
+	return FReply::Unhandled();
+}
+
+FSearchEntry * FSearchEntry::MakeCategoryEntry(const FString & InTitle)
+{
+	FSearchEntry * SearchEntry = new FSearchEntry();
+	SearchEntry->Title = InTitle;
+	SearchEntry->bCategory = true;
+
+	return SearchEntry;
+}
+
+void UpdateSuggestionHelper(const FText & CategoryLabel, const TArray<FSearchEntry> & Elements, TArray<TSharedPtr< FSearchEntry > > & OutSuggestions)
+{
+	if (Elements.Num())
+	{
+		OutSuggestions.Add(MakeShareable(FSearchEntry::MakeCategoryEntry(CategoryLabel.ToString())));
+	}
+
+	for (uint32 i = 0; i < (uint32)Elements.Num(); ++i)
+	{
+		OutSuggestions.Add(MakeShareable(new FSearchEntry(Elements[i])));
+	}
+}
+
+void SSuperSearchBox::UpdateSuggestions()
+{
+	const FText & Query = InputText->GetText();
+	FSearchResults * SearchResults = SearchResultsCache.Find(Query.ToString());
+
+	//still waiting on results for current query
+	if (SearchResults == NULL)
+	{
+		return;
+	}
+
+	//go through and build new suggestion list for list view widget
+	Suggestions.Empty();
+	SelectedSuggestion = -1;
+	
+	//first show api
+	UpdateSuggestionHelper(NSLOCTEXT("SuperSearch", "API", "API"), SearchResults->OnlineResults.FindOrAdd(TEXT("api")), Suggestions);
+
+	//then docs
+	UpdateSuggestionHelper(NSLOCTEXT("SuperSearch", "docs", "Documentation"), SearchResults->OnlineResults.FindOrAdd(TEXT("documentation")), Suggestions);
+
+	//then wiki
+	UpdateSuggestionHelper(NSLOCTEXT("SuperSearch", "wiki", "Wiki"), SearchResults->OnlineResults.FindOrAdd(TEXT("wiki")), Suggestions);
+
+	//then answerhub
+	UpdateSuggestionHelper(NSLOCTEXT("SuperSearch", "answers", "Answerhub"), SearchResults->OnlineResults.FindOrAdd(TEXT("answers")), Suggestions);
+
+	if(Suggestions.Num())
+	{
+		SelectedSuggestion = 1;
+		// Ideally if the selection box is open the output window is not changing it's window title (flickers)
+		SuggestionBox->SetIsOpen(true, false);
+		MarkActiveSuggestion();
+
+		// Force the textbox back into focus.
+		FWidgetPath WidgetToFocusPath;
+		FSlateApplication::Get().GeneratePathToWidgetUnchecked( InputText.ToSharedRef(), WidgetToFocusPath );
+		FSlateApplication::Get().SetKeyboardFocus( WidgetToFocusPath, EKeyboardFocusCause::SetDirectly );
+	}
+	else
+	{
+		SuggestionBox->SetIsOpen(false);
+	}
+}
+
+void SSuperSearchBox::OnKeyboardFocusLost( const FKeyboardFocusEvent& InKeyboardFocusEvent )
+{
+//	SuggestionBox->SetIsOpen(false);
+}
+
+void SSuperSearchBox::MarkActiveSuggestion()
+{
+	bIgnoreUIUpdate = true;
+	if(SelectedSuggestion >= 1)
+	{
+		SuggestionListView->SetSelection(Suggestions[SelectedSuggestion]);
+		SuggestionListView->RequestScrollIntoView(Suggestions[SelectedSuggestion]);	// Ideally this would only scroll if outside of the view
+	}
+	else
+	{
+		SuggestionListView->ClearSelection();
+	}
+	bIgnoreUIUpdate = false;
+}
+
+void SSuperSearchBox::ClearSuggestions()
+{
+	SelectedSuggestion = -1;
+	SuggestionBox->SetIsOpen(false);
+	Suggestions.Empty();
+}
+
+FString SSuperSearchBox::GetSelectionText() const
+{
+	FString ret = Suggestions[SelectedSuggestion]->Title;
+	
+	ret = ret.Replace(TEXT("\t"), TEXT(""));
+
+	return ret;
+}
\ No newline at end of file
diff --git a/Engine/Source/Developer/SuperSearch/Private/SSuperSearch.h b/Engine/Source/Developer/SuperSearch/Private/SSuperSearch.h
new file mode 100644
--- /dev/null
+++ b/Engine/Source/Developer/SuperSearch/Private/SSuperSearch.h
@@ -0,0 +1,117 @@
+// Copyright 1998-2014 Epic Games, Inc. All Rights Reserved.
+
+#pragma once
+
+/**
+ * Console input box with command-completion support
+ */
+
+struct FSearchEntry
+{
+	FString Title;
+	FString URL;
+	bool bCategory;
+
+	static FSearchEntry * MakeCategoryEntry(const FString & InTitle);
+};
+
+
+class SSuperSearchBox
+	: public SCompoundWidget
+{
+
+public:
+
+	SLATE_BEGIN_ARGS( SSuperSearchBox )
+		: _SuggestionListPlacement( MenuPlacement_BelowAnchor )
+		{}
+
+		/** Where to place the suggestion list */
+		SLATE_ARGUMENT( EMenuPlacement, SuggestionListPlacement )
+
+	SLATE_END_ARGS()
+
+	/** Protected console input box widget constructor, called by Slate */
+	SSuperSearchBox();
+
+	/**
+	 * Construct this widget.  Called by the SNew() Slate macro.
+	 *
+	 * @param	InArgs	Declaration used by the SNew() macro to construct this widget
+	 */
+	void Construct( const FArguments& InArgs );
+
+	/** Returns the editable text box associated with this widget.  Used to set focus directly. */
+	TSharedRef< SEditableTextBox > GetEditableTextBox()
+	{
+		return InputText.ToSharedRef();
+	}
+
+	/** SWidget interface */
+	virtual void Tick( const FGeometry& AllottedGeometry, const double InCurrentTime, const float InDeltaTime ) override;
+
+protected:
+
+	virtual bool SupportsKeyboardFocus() const { return true; }
+
+	// e.g. Tab or Key_Up
+	virtual FReply OnKeyDown( const FGeometry& MyGeometry, const FKeyboardEvent& KeyboardEvent );
+
+	void OnKeyboardFocusLost( const FKeyboardFocusEvent& InKeyboardFocusEvent );
+
+	/** Handles entering in a command */
+	void OnTextCommitted(const FText& InText, ETextCommit::Type CommitInfo);
+
+	void OnTextChanged(const FText& InText);
+
+	/** Makes the widget for the suggestions messages in the list view */
+	TSharedRef<ITableRow> MakeSuggestionListItemWidget(TSharedPtr<FSearchEntry> Suggestion, const TSharedRef<STableViewBase>& OwnerTable);
+
+	void Query_HttpRequestComplete(FHttpRequestPtr HttpRequest, FHttpResponsePtr HttpResponse, bool bSucceeded);
+
+	void SuggestionSelectionChanged(TSharedPtr<FSearchEntry> NewValue, ESelectInfo::Type SelectInfo);
+	
+	void ActOnSuggestion(TSharedPtr<FSearchEntry> SearchEntry);
+		
+	void UpdateSuggestions();
+
+	void MarkActiveSuggestion();
+
+	void ClearSuggestions();
+
+	FString GetSelectionText() const;
+
+
+private:
+
+	/** Editable text widget */
+	TSharedPtr< SEditableTextBox > InputText;
+
+	/** history / auto completion elements */
+	TSharedPtr< SMenuAnchor > SuggestionBox;
+
+	/** All suggestions stored in this widget for the list view */
+	TArray< TSharedPtr<FSearchEntry> > Suggestions;
+
+	/** The list view for showing all log messages. Should be replaced by a full text editor */
+	TSharedPtr< SListView< TSharedPtr<FSearchEntry> > > SuggestionListView;
+
+	//TODO: properly clean map up as time goes by requests are processed and can be discarded
+	TMap<FHttpRequestPtr, FText> RequestQueryMap;
+
+	typedef TMap<FString, TArray<FSearchEntry> > FCategoryResults;
+	struct FSearchResults
+	{
+		 FCategoryResults OnlineResults;
+	};
+
+	//TODO: properly clean map up as time goes by requests are processed and can be discarded
+	/** A cache of results based on search */
+	TMap<FString, FSearchResults> SearchResultsCache;
+
+	/** -1 if not set, otherwise index into Suggestions */
+	int32 SelectedSuggestion;
+
+	/** to prevent recursive calls in UI callback */
+	bool bIgnoreUIUpdate; 
+};
\ No newline at end of file
diff --git a/Engine/Source/Developer/SuperSearch/Private/SuperSearchModule.cpp b/Engine/Source/Developer/SuperSearch/Private/SuperSearchModule.cpp
new file mode 100644
--- /dev/null
+++ b/Engine/Source/Developer/SuperSearch/Private/SuperSearchModule.cpp
@@ -0,0 +1,26 @@
+// Copyright 1998-2014 Epic Games, Inc. All Rights Reserved.
+
+#include "SuperSearchPrivatePCH.h"
+#include "SSuperSearch.h"
+
+IMPLEMENT_MODULE( FSuperSearchModule,SuperSearch );
+
+namespace SuperSearchModule
+{
+	static const FName FNameSuperSearchApp = FName(TEXT("SuperSearchApp"));
+}
+
+void FSuperSearchModule::StartupModule()
+{
+}
+
+void FSuperSearchModule::ShutdownModule()
+{
+}
+
+TSharedRef< SWidget > FSuperSearchModule::MakeSearchBox( TSharedPtr< SEditableTextBox >& OutExposedEditableTextBox ) const
+{
+	TSharedRef< SSuperSearchBox > NewSearchBox = SNew( SSuperSearchBox );
+	OutExposedEditableTextBox = NewSearchBox->GetEditableTextBox();
+	return NewSearchBox;
+}
diff --git a/Engine/Source/Developer/SuperSearch/Private/SuperSearchPrivatePCH.h b/Engine/Source/Developer/SuperSearch/Private/SuperSearchPrivatePCH.h
new file mode 100644
--- /dev/null
+++ b/Engine/Source/Developer/SuperSearch/Private/SuperSearchPrivatePCH.h
@@ -0,0 +1,10 @@
+// Copyright 1998-2014 Epic Games, Inc. All Rights Reserved.
+
+#pragma once
+
+#include "UnrealEd.h"
+#include "SuperSearchModule.h"
+#include "Slate.h"
+#include "Http.h"
+#include "Json.h"
+
diff --git a/Engine/Source/Developer/SuperSearch/Public/SuperSearchModule.h b/Engine/Source/Developer/SuperSearch/Public/SuperSearchModule.h
new file mode 100644
--- /dev/null
+++ b/Engine/Source/Developer/SuperSearch/Public/SuperSearchModule.h
@@ -0,0 +1,16 @@
+// Copyright 1998-2014 Epic Games, Inc. All Rights Reserved.
+
+#pragma once
+
+#include "Slate.h"
+
+class FSuperSearchModule : public IModuleInterface
+{
+public:
+	virtual void StartupModule();
+	virtual void ShutdownModule();
+
+	/** Generates aSuperSearch box widget.  Remember, this widget will become invalid if the
+	   SuperSearch DLL is unloaded on the fly. */
+	virtual TSharedRef< SWidget > MakeSearchBox( TSharedPtr< SEditableTextBox >& OutExposedEditableTextBox ) const;
+};
diff --git a/Engine/Source/Developer/SuperSearch/SuperSearch.Build.cs b/Engine/Source/Developer/SuperSearch/SuperSearch.Build.cs
new file mode 100644
--- /dev/null
+++ b/Engine/Source/Developer/SuperSearch/SuperSearch.Build.cs
@@ -0,0 +1,23 @@
+// Copyright 1998-2014 Epic Games, Inc. All Rights Reserved.
+
+using UnrealBuildTool;
+
+public class SuperSearch : ModuleRules
+{
+	public SuperSearch(TargetInfo Target)
+	{
+		PrivateDependencyModuleNames.AddRange(
+			new string[] {
+				"Core",
+				"CoreUObject", // @todo Mac: for some reason it's needed to link in debug on Mac
+				"Engine",
+                "InputCore",
+				"UnrealEd",
+				"Slate",
+				"SlateCore",
+                "EditorStyle",
+                "HTTP"
+			}
+		);
+	}
+}
diff --git a/Engine/Source/Editor/MainFrame/MainFrame.Build.cs b/Engine/Source/Editor/MainFrame/MainFrame.Build.cs
--- a/Engine/Source/Editor/MainFrame/MainFrame.Build.cs
+++ b/Engine/Source/Editor/MainFrame/MainFrame.Build.cs
@@ -47,7 +47,8 @@ public class MainFrame : ModuleRules
 				"GameProjectGeneration",
 				"ProjectTargetPlatformEditor",
 				"LevelEditor",
-				"OutputLog",
+                "OutputLog",
+				"SuperSearch",
 				"Settings",
 				"SourceCodeAccess",
 			}
@@ -69,7 +70,8 @@ public class MainFrame : ModuleRules
 				"GameProjectGeneration",
 				"ProjectTargetPlatformEditor",
 				"LevelEditor",
-				"OutputLog",
+                "OutputLog",
+				"SuperSearch",
                 "TranslationEditor",
 				"SourceCodeAccess",
 				"EditorLiveStreaming"
diff --git a/Engine/Source/Editor/MainFrame/Private/MainFrameModule.cpp b/Engine/Source/Editor/MainFrame/Private/MainFrameModule.cpp
--- a/Engine/Source/Editor/MainFrame/Private/MainFrameModule.cpp
+++ b/Engine/Source/Editor/MainFrame/Private/MainFrameModule.cpp
@@ -316,8 +316,11 @@ TSharedRef<SWidget> FMainFrameModule::MakeDeveloperTools() const
 	};
 
 
+	
+	const FSuperSearchModule& SuperSearchModule = FModuleManager::LoadModuleChecked< FSuperSearchModule >(TEXT("SuperSearch"));
+	
 	// We need the output log module in order to instantiate SConsoleInputBox widgets
-	const FOutputLogModule& OutputLogModule = FModuleManager::LoadModuleChecked< FOutputLogModule >( TEXT( "OutputLog" ) );
+	const FOutputLogModule& OutputLogModule = FModuleManager::LoadModuleChecked< FOutputLogModule >(TEXT("OutputLog"));
 
 	const FSlateFontInfo& SmallFixedFont = FEditorStyle::GetFontStyle(TEXT("MainFrame.DebugTools.SmallFont") );
 	const FSlateFontInfo& NormalFixedFont = FEditorStyle::GetFontStyle(TEXT("MainFrame.DebugTools.NormalFont") );
@@ -439,6 +442,8 @@ TSharedRef<SWidget> FMainFrameModule::MakeDeveloperTools() const
 		]
 	;
 
+	bool bUseSuperSearch = FParse::Param(FCommandLine::Get(), TEXT("SuperSearch"));;
+
 	// Invisible border, so that we can animate our box panel size
 	return SNew( SBorder )
 		.Visibility( EVisibility::SelfHitTestInvisible )
@@ -465,7 +470,7 @@ TSharedRef<SWidget> FMainFrameModule::MakeDeveloperTools() const
 					.Padding( FMargin( 4.0f, 0.0f, 0.0f, 0.0f ) )
 					.WidthOverride( 180.0f )
 					[
-						OutputLogModule.MakeConsoleInputBox( ExposedEditableTextBox )
+						bUseSuperSearch ? SuperSearchModule.MakeSearchBox( ExposedEditableTextBox ) : OutputLogModule.MakeConsoleInputBox( ExposedEditableTextBox )
 					]
 				]
 /*
diff --git a/Engine/Source/Editor/MainFrame/Private/MainFramePrivatePCH.h b/Engine/Source/Editor/MainFrame/Private/MainFramePrivatePCH.h
--- a/Engine/Source/Editor/MainFrame/Private/MainFramePrivatePCH.h
+++ b/Engine/Source/Editor/MainFrame/Private/MainFramePrivatePCH.h
@@ -24,6 +24,7 @@
 #include "MessageLogModule.h"
 #include "ModuleManager.h"
 #include "MRUFavoritesList.h"
+#include "SuperSearchModule.h"
 #include "OutputLogModule.h"
 //#include "SearchUI.h"
 #include "SourceCodeNavigation.h"
