SuperSearch clears text when loses focus.
Online search is only done when typing stops

diff --git a/Engine/Source/Developer/SuperSearch/Private/SSuperSearch.cpp b/Engine/Source/Developer/SuperSearch/Private/SSuperSearch.cpp
--- a/Engine/Source/Developer/SuperSearch/Private/SSuperSearch.cpp
+++ b/Engine/Source/Developer/SuperSearch/Private/SSuperSearch.cpp
@@ -20,9 +20,10 @@ void SSuperSearchBox::Construct( const FArguments& InArgs )
 			[
 				SAssignNew(InputText, SSearchBox)
 					.OnTextCommitted(this, &SSuperSearchBox::OnTextCommitted)
-					.HintText( NSLOCTEXT( "SuperSearchBox", "HelpHint", "Help" ) )
+					.HintText( NSLOCTEXT( "SuperSearchBox", "HelpHint", "Search For Help" ) )
 					.OnTextChanged(this, &SSuperSearchBox::OnTextChanged)
 					.SelectAllTextWhenFocused(false)
+					.DelayChangeNotificationsWhileTyping(true)
 			]
 			.MenuContent
 			(
@@ -42,6 +43,7 @@ void SSuperSearchBox::Construct( const FArguments& InArgs )
 					]
 				]
 			)
+			.OnMenuOpenChanged(this, &SSuperSearchBox::OnMenuOpenChanged)
 	];
 }
 END_SLATE_FUNCTION_BUILD_OPTIMIZATION
@@ -60,6 +62,14 @@ void SSuperSearchBox::Tick( const FGeometry& AllottedGeometry, const double InCu
 	}
 }
 
+void SSuperSearchBox::OnMenuOpenChanged( bool bIsOpen )
+{
+	if (bIsOpen == false)
+	{
+		InputText->SetText(FText::GetEmpty());
+	}
+}
+
 void SSuperSearchBox::ActOnSuggestion(TSharedPtr<FSearchEntry> SearchEntry)
 {
 	if (SearchEntry->bCategory == false)
diff --git a/Engine/Source/Developer/SuperSearch/Private/SSuperSearch.h b/Engine/Source/Developer/SuperSearch/Private/SSuperSearch.h
--- a/Engine/Source/Developer/SuperSearch/Private/SSuperSearch.h
+++ b/Engine/Source/Developer/SuperSearch/Private/SSuperSearch.h
@@ -71,6 +71,8 @@ protected:
 
 	void SuggestionSelectionChanged(TSharedPtr<FSearchEntry> NewValue, ESelectInfo::Type SelectInfo);
 	
+	void OnMenuOpenChanged(bool bIsOpen);
+
 	void ActOnSuggestion(TSharedPtr<FSearchEntry> SearchEntry);
 		
 	void UpdateSuggestions();
