Fix supersearch crash when clicking on empty space
TTP#346115

diff --git a/Engine/Source/Developer/SuperSearch/Private/SSuperSearch.cpp b/Engine/Source/Developer/SuperSearch/Private/SSuperSearch.cpp
--- a/Engine/Source/Developer/SuperSearch/Private/SSuperSearch.cpp
+++ b/Engine/Source/Developer/SuperSearch/Private/SSuperSearch.cpp
@@ -85,7 +85,7 @@ void SSuperSearchBox::SuggestionSelectionChanged(TSharedPtr<FSearchEntry> NewVal
 		return;
 	}
 
-	if (NewValue->bCategory)
+	if (NewValue.Get() == NULL || NewValue->bCategory)
 	{
 		return;
 	}
