visual polish on super search

diff --git a/Engine/Source/Developer/SuperSearch/Private/SSuperSearch.cpp b/Engine/Source/Developer/SuperSearch/Private/SSuperSearch.cpp
--- a/Engine/Source/Developer/SuperSearch/Private/SSuperSearch.cpp
+++ b/Engine/Source/Developer/SuperSearch/Private/SSuperSearch.cpp
@@ -8,6 +8,11 @@ SSuperSearchBox::SSuperSearchBox()
 	: SelectedSuggestion(-1)
 	, bIgnoreUIUpdate(false)
 {
+	CategoryToIconMap.Add("Documentation", FName("LevelEditor.BrowseDocumentation") );
+	CategoryToIconMap.Add("Wiki", FName("MainFrame.VisitWiki") );
+	CategoryToIconMap.Add("Answerhub", FName("MainFrame.VisitSearchForAnswersPage") );
+	CategoryToIconMap.Add("API", FName("LevelEditor.BrowseAPIReference") );
+	CategoryToIconMap.Add("Tutorials", FName("LevelEditor.Tutorials") );
 }
 
 BEGIN_SLATE_FUNCTION_BUILD_OPTIMIZATION
@@ -16,52 +21,39 @@ void SSuperSearchBox::Construct( const FArguments& InArgs )
 	ChildSlot
 	[
 		SAssignNew( SuggestionBox, SMenuAnchor )
-			.Placement( InArgs._SuggestionListPlacement )
+		.Placement( InArgs._SuggestionListPlacement )
+		[
+			SAssignNew(InputText, SSearchBox)
+			.OnTextCommitted(this, &SSuperSearchBox::OnTextCommitted)
+			.HintText(NSLOCTEXT("SuperSearchBox", "HelpHint", "Search For Help"))
+			.OnTextChanged(this, &SSuperSearchBox::OnTextChanged)
+			.SelectAllTextWhenFocused(false)
+			.DelayChangeNotificationsWhileTyping(true)
+			.MinDesiredWidth(200)
+		]
+		.MenuContent
+		(
+			SNew(SBorder)
+			.BorderImage(FEditorStyle::GetBrush("Menu.Background"))
+			.Padding( FMargin(2) )
 			[
-				SAssignNew(InputText, SSearchBox)
-					.OnTextCommitted(this, &SSuperSearchBox::OnTextCommitted)
-					.HintText( NSLOCTEXT( "SuperSearchBox", "HelpHint", "Search For Help" ) )
-					.OnTextChanged(this, &SSuperSearchBox::OnTextChanged)
-					.SelectAllTextWhenFocused(false)
-					.DelayChangeNotificationsWhileTyping(true)
-			]
-			.MenuContent
-			(
-				SNew(SBorder)
-				.BorderImage(FEditorStyle::GetBrush("Menu.Background"))
-				.Padding( FMargin(2) )
+				SNew(SBox)
+				.HeightOverride(250) // avoids flickering, ideally this would be adaptive to the content without flickering
 				[
-					SNew(SBox)
-					.HeightOverride(250) // avoids flickering, ideally this would be adaptive to the content without flickering
-					[
-						SAssignNew(SuggestionListView, SListView< TSharedPtr<FSearchEntry> >)
-							.ListItemsSource(&Suggestions)
-							.SelectionMode( ESelectionMode::Single )							// Ideally the mouse over would not highlight while keyboard controls the UI
-							.OnGenerateRow(this, &SSuperSearchBox::MakeSuggestionListItemWidget)
-							.OnSelectionChanged(this, &SSuperSearchBox::SuggestionSelectionChanged)
-							.ItemHeight(18)
-					]
+					SAssignNew(SuggestionListView, SListView< TSharedPtr<FSearchEntry> >)
+						.ListItemsSource(&Suggestions)
+						.SelectionMode( ESelectionMode::Single )							// Ideally the mouse over would not highlight while keyboard controls the UI
+						.OnGenerateRow(this, &SSuperSearchBox::MakeSuggestionListItemWidget)
+						.OnSelectionChanged(this, &SSuperSearchBox::SuggestionSelectionChanged)
+						.ItemHeight(18)
 				]
-			)
-			.OnMenuOpenChanged(this, &SSuperSearchBox::OnMenuOpenChanged)
+			]
+		)
+		.OnMenuOpenChanged(this, &SSuperSearchBox::OnMenuOpenChanged)
 	];
 }
 END_SLATE_FUNCTION_BUILD_OPTIMIZATION
 
-void SSuperSearchBox::Tick( const FGeometry& AllottedGeometry, const double InCurrentTime, const float InDeltaTime )
-{
-	SCompoundWidget::Tick(AllottedGeometry, InCurrentTime, InDeltaTime);
-
-	if (!GIntraFrameDebuggingGameThread && !IsEnabled())
-	{
-		SetEnabled(true);
-	}
-	else if (GIntraFrameDebuggingGameThread && IsEnabled())
-	{
-		SetEnabled(false);
-	}
-}
-
 void SSuperSearchBox::OnMenuOpenChanged( bool bIsOpen )
 {
 	if (bIsOpen == false)
@@ -123,23 +115,37 @@ TSharedRef<ITableRow> SSuperSearchBox::MakeSuggestionListItemWidget(TSharedPtr<F
 		Combined = Suggestion->Title;
 	}
 
-	FText HighlightText = FText::FromString(Left);
-
 	bool bCategory = Suggestion->bCategory;
 
+	
+	TSharedPtr<SWidget> IconWidget;
+	if (bCategory)
+	{
+		FName * IconResult = CategoryToIconMap.Find(Combined);
+		SAssignNew(IconWidget, SImage)
+			.Image(FEditorStyle::GetBrush(IconResult ? *IconResult : FName("MainFrame.VisitEpicGamesDotCom") ));
+	}
+
 	return
 		SNew(STableRow< TSharedPtr<FString> >, OwnerTable)
 		.ShowSelection(bCategory == false)
+		.Padding(FMargin(bCategory ? 2.5 : 16, 2.5, 2.5, 2.5))
 		[
-			SNew(SBox)
-			.WidthOverride(600)			// to enforce some minimum width, ideally we define the minimum, not a fixed width
-			.Padding(FMargin(bCategory ? 0 : 16, 0, 0, 0))
+			SNew(SHorizontalBox)
+			+ SHorizontalBox::Slot()
+			.Padding(FMargin(0,0,2.5,0))
+			[
+				bCategory ? IconWidget.ToSharedRef() : SNew(SSpacer)
+			]
+			+SHorizontalBox::Slot()
+			.AutoWidth()
 			[
 				SNew(STextBlock)
 				.Text(Combined)
-				//.TextStyle( FEditorStyle::Get(), "Log.Normal")
-				.HighlightText(HighlightText)
+				.TextStyle(&FEditorStyle::Get().GetWidgetStyle<FTextBlockStyle>(bCategory ? "SuperSearchCategoryText" : "NormalText"))
+				.MinDesiredWidth(600)
 			]
+			
 		];
 }
 
diff --git a/Engine/Source/Developer/SuperSearch/Private/SSuperSearch.h b/Engine/Source/Developer/SuperSearch/Private/SSuperSearch.h
--- a/Engine/Source/Developer/SuperSearch/Private/SSuperSearch.h
+++ b/Engine/Source/Developer/SuperSearch/Private/SSuperSearch.h
@@ -47,9 +47,6 @@ public:
 		return InputText.ToSharedRef();
 	}
 
-	/** SWidget interface */
-	virtual void Tick( const FGeometry& AllottedGeometry, const double InCurrentTime, const float InDeltaTime ) override;
-
 protected:
 
 	virtual bool SupportsKeyboardFocus() const { return true; }
@@ -111,6 +108,8 @@ private:
 	/** A cache of results based on search */
 	TMap<FString, FSearchResults> SearchResultsCache;
 
+	TMap<FString, FName> CategoryToIconMap;
+
 	/** -1 if not set, otherwise index into Suggestions */
 	int32 SelectedSuggestion;
 
diff --git a/Engine/Source/Editor/EditorStyle/Private/SlateEditorStyle.cpp b/Engine/Source/Editor/EditorStyle/Private/SlateEditorStyle.cpp
--- a/Engine/Source/Editor/EditorStyle/Private/SlateEditorStyle.cpp
+++ b/Engine/Source/Editor/EditorStyle/Private/SlateEditorStyle.cpp
@@ -2156,6 +2156,13 @@ void FSlateEditorStyle::FStyle::SetupGeneralStyles()
 #endif // WITH_EDITOR || IS_PROGRAM
 
 #if WITH_EDITOR
+	// Supersearch
+	FTextBlockStyle SuperSearchCategoryText = FTextBlockStyle(NormalText)
+		.SetFont(TTF_CORE_FONT("Fonts/Roboto-Bold", 10))
+		.SetColorAndOpacity(FLinearColor(0, 0.3f, 1.f));
+
+	Set("SuperSearchCategoryText", SuperSearchCategoryText);
+
 	// Gameplay Tags
 	{
 		Set("GameplayTagTreeView", FTableRowStyle()
diff --git a/Engine/Source/Editor/MainFrame/Private/MainFrameModule.cpp b/Engine/Source/Editor/MainFrame/Private/MainFrameModule.cpp
--- a/Engine/Source/Editor/MainFrame/Private/MainFrameModule.cpp
+++ b/Engine/Source/Editor/MainFrame/Private/MainFrameModule.cpp
@@ -484,7 +484,6 @@ TSharedRef<SWidget> FMainFrameModule::MakeDeveloperTools() const
 				[
 					SNew(SBox)
 					.Padding( FMargin( 4.0f, 0.0f, 0.0f, 0.0f ) )
-					.WidthOverride( 180.0f )
 					[
 						bUseSuperSearch ? SuperSearchModule.MakeSearchBox( ExposedEditableTextBox ) : OutputLogModule.MakeConsoleInputBox( ExposedEditableTextBox )
 					]
