Add analytics to supersearch

diff --git a/Engine/Source/Developer/SuperSearch/Private/SSuperSearch.cpp b/Engine/Source/Developer/SuperSearch/Private/SSuperSearch.cpp
--- a/Engine/Source/Developer/SuperSearch/Private/SSuperSearch.cpp
+++ b/Engine/Source/Developer/SuperSearch/Private/SSuperSearch.cpp
@@ -54,18 +54,46 @@ void SSuperSearchBox::Construct( const FArguments& InArgs )
 }
 END_SLATE_FUNCTION_BUILD_OPTIMIZATION
 
+int32 GetNumRealSuggestions(const TArray< TSharedPtr<FSearchEntry> > & Suggestions)
+{
+	int32 NumSuggestions = 0;
+	for (auto Entry : Suggestions)
+	{
+		if (Entry->bCategory == false)
+		{
+			++NumSuggestions;
+		}
+	}
+
+	return NumSuggestions;
+}
+
 void SSuperSearchBox::OnMenuOpenChanged( bool bIsOpen )
 {
 	if (bIsOpen == false)
 	{
+		if(FEngineAnalytics::IsAvailable())
+		{
+			TArray< FAnalyticsEventAttribute > SearchAttribs;
+			SearchAttribs.Add(FAnalyticsEventAttribute(TEXT("SearchTerm"), InputText->GetText().ToString() ));
+			SearchAttribs.Add(FAnalyticsEventAttribute(TEXT("NumFoundResults"), GetNumRealSuggestions(Suggestions)));
+			SearchAttribs.Add(FAnalyticsEventAttribute(TEXT("ClickedTitle"), EntryClicked.Get() ? EntryClicked->Title : "") );
+			SearchAttribs.Add(FAnalyticsEventAttribute(TEXT("ClickedOnSomething"), EntryClicked.IsValid()) );
+
+			FEngineAnalytics::GetProvider().RecordEvent(TEXT("Editor.Usage.SuperSearch"), SearchAttribs);
+		}
+
 		InputText->SetText(FText::GetEmpty());
 	}
+
+	EntryClicked = NULL;
 }
 
 void SSuperSearchBox::ActOnSuggestion(TSharedPtr<FSearchEntry> SearchEntry)
 {
 	if (SearchEntry->bCategory == false)
 	{
+		EntryClicked = SearchEntry;
 		FPlatformProcess::LaunchURL(*SearchEntry->URL, NULL, NULL);
 	}
 }
diff --git a/Engine/Source/Developer/SuperSearch/Private/SSuperSearch.h b/Engine/Source/Developer/SuperSearch/Private/SSuperSearch.h
--- a/Engine/Source/Developer/SuperSearch/Private/SSuperSearch.h
+++ b/Engine/Source/Developer/SuperSearch/Private/SSuperSearch.h
@@ -113,6 +113,9 @@ private:
 	/** -1 if not set, otherwise index into Suggestions */
 	int32 SelectedSuggestion;
 
+	/** Entry that was actually clicked or pressed enter on. NULLs out when menu closes */
+	TSharedPtr<FSearchEntry> EntryClicked;
+
 	/** to prevent recursive calls in UI callback */
 	bool bIgnoreUIUpdate; 
 };
\ No newline at end of file
