Fix clang conditional lvalue

diff --git a/Engine/Source/Developer/SuperSearch/Private/SSuperSearch.cpp b/Engine/Source/Developer/SuperSearch/Private/SSuperSearch.cpp
--- a/Engine/Source/Developer/SuperSearch/Private/SSuperSearch.cpp
+++ b/Engine/Source/Developer/SuperSearch/Private/SSuperSearch.cpp
@@ -153,6 +153,10 @@ TSharedRef<ITableRow> SSuperSearchBox::MakeSuggestionListItemWidget(TSharedPtr<F
 		SAssignNew(IconWidget, SImage)
 			.Image(FEditorStyle::GetBrush(IconResult ? *IconResult : FName("MainFrame.VisitEpicGamesDotCom") ));
 	}
+	else
+	{
+		SAssignNew(IconWidget, SSpacer);
+	}
 
 	return
 		SNew(STableRow< TSharedPtr<FString> >, OwnerTable)
@@ -163,7 +167,7 @@ TSharedRef<ITableRow> SSuperSearchBox::MakeSuggestionListItemWidget(TSharedPtr<F
 			+ SHorizontalBox::Slot()
 			.Padding(FMargin(0,0,2.5,0))
 			[
-				bCategory ? IconWidget.ToSharedRef() : SNew(SSpacer)
+				IconWidget.ToSharedRef()
 			]
 			+SHorizontalBox::Slot()
 			.AutoWidth()
