Expose tutorials to super search

diff --git a/Engine/Source/Developer/SuperSearch/Private/SSuperSearch.cpp b/Engine/Source/Developer/SuperSearch/Private/SSuperSearch.cpp
--- a/Engine/Source/Developer/SuperSearch/Private/SSuperSearch.cpp
+++ b/Engine/Source/Developer/SuperSearch/Private/SSuperSearch.cpp
@@ -3,6 +3,9 @@
 #include "SuperSearchPrivatePCH.h"
 #include "SSuperSearch.h"
 #include "SScrollBorder.h"
+#include "AssetRegistryModule.h"
+#include "IIntroTutorials.h"
+#include "EditorTutorial.h"
 
 SSuperSearchBox::SSuperSearchBox()
 	: SelectedSuggestion(-1)
@@ -44,7 +47,7 @@ void SSuperSearchBox::Construct( const FArguments& InArgs )
 						.ListItemsSource(&Suggestions)
 						.SelectionMode( ESelectionMode::Single )							// Ideally the mouse over would not highlight while keyboard controls the UI
 						.OnGenerateRow(this, &SSuperSearchBox::MakeSuggestionListItemWidget)
-						.OnSelectionChanged(this, &SSuperSearchBox::SuggestionSelectionChanged)
+						.OnMouseButtonClick(this, &SSuperSearchBox::SuggestionSelectionChanged)
 						.ItemHeight(18)
 				]
 			]
@@ -94,11 +97,27 @@ void SSuperSearchBox::ActOnSuggestion(TSharedPtr<FSearchEntry> SearchEntry)
 	if (SearchEntry->bCategory == false)
 	{
 		EntryClicked = SearchEntry;
-		FPlatformProcess::LaunchURL(*SearchEntry->URL, NULL, NULL);
+		if (SearchEntry->URL.IsEmpty())
+		{
+			UBlueprint* Blueprint = LoadObject<UBlueprint>(nullptr, *(SearchEntry->AssetData.ObjectPath.ToString()));
+			UEditorTutorial* Tutorial = Blueprint->GeneratedClass->GetDefaultObject<UEditorTutorial>();
+			IIntroTutorials& IntroTutorials = IIntroTutorials::Get();
+			
+			FWidgetPath OutWidgetPath;
+			TSharedPtr<SWindow> ParentWindow = FSlateApplication::Get().FindWidgetWindow(SuggestionBox.ToSharedRef(), OutWidgetPath);
+
+			IntroTutorials.LaunchTutorial(Tutorial, true, ParentWindow);
+		}
+		else
+		{
+			FPlatformProcess::LaunchURL(*SearchEntry->URL, NULL, NULL);
+		}
+
+		SuggestionBox->SetIsOpen(false);
 	}
 }
 
-void SSuperSearchBox::SuggestionSelectionChanged(TSharedPtr<FSearchEntry> NewValue, ESelectInfo::Type SelectInfo)
+void SSuperSearchBox::SuggestionSelectionChanged(TSharedPtr<FSearchEntry> NewValue)
 {
 	if(bIgnoreUIUpdate)
 	{
@@ -117,10 +136,7 @@ void SSuperSearchBox::SuggestionSelectionChanged(TSharedPtr<FSearchEntry> NewVal
 			SelectedSuggestion = i;
 			MarkActiveSuggestion();
 
-			if( SelectInfo == ESelectInfo::OnMouseClick )
-			{
-				ActOnSuggestion(NewValue);
-			}
+			ActOnSuggestion(NewValue);
 
 			break;
 		}
@@ -191,23 +207,61 @@ void SSuperSearchBox::OnTextChanged(const FText& InText)
 	const FString& InputTextStr = InputText->GetText().ToString();
 	if(!InputTextStr.IsEmpty())
 	{
+		//search through google for online content
+		{
+			TSharedRef<class IHttpRequest> HttpRequest = FHttpModule::Get().CreateRequest();
+
+			// build the url
+			// use partial response to only look at items and things we care about for items right now (title,link and label name)
+			const FText QueryURL = FText::Format(FText::FromString("https://www.googleapis.com/customsearch/v1?key=AIzaSyCMGfdDaSfjqv5zYoS0mTJnOT3e9MURWkU&cx=009868829633250020713:y7tfd8hlcgg&fields=items(title,link,labels/name)&q={0}+less:forums"), InText);
+
+			//save http request into map to ensure correct ordering
+			FText & Query = RequestQueryMap.FindOrAdd(HttpRequest);
+			Query = InText;
+
+			// kick off http request to read friends
+			HttpRequest->OnProcessRequestComplete().BindRaw(this, &SSuperSearchBox::Query_HttpRequestComplete);
+			HttpRequest->SetURL(QueryURL.ToString());
+			HttpRequest->SetHeader(TEXT("Content-Type"), TEXT("application/json"));
+			HttpRequest->SetVerb(TEXT("GET"));
+			HttpRequest->ProcessRequest();
+		}
+
+		//search local tutorials
+		{
+			FAssetRegistryModule& AssetRegistry = FModuleManager::LoadModuleChecked<FAssetRegistryModule>(TEXT("AssetRegistry"));
+
+			FARFilter Filter;
+			Filter.ClassNames.Add(UBlueprint::StaticClass()->GetFName());
+			Filter.bRecursiveClasses = true;
+			Filter.TagsAndValues.Add(TEXT("ParentClass"), FString::Printf(TEXT("%s'%s'"), *UClass::StaticClass()->GetName(), *UEditorTutorial::StaticClass()->GetPathName()));
+
+			TArray<FAssetData> AssetData;
+			AssetRegistry.Get().GetAssets(Filter, AssetData);
+
+			//add search result into cache
+			FSearchResults& SearchResults = SearchResultsCache.FindOrAdd(InText.ToString());
+			TArray<FSearchEntry>& TutorialResults = SearchResults.TutorialResults;
+			TutorialResults.Empty();
+
+			for (const FAssetData& Asset : AssetData)
+			{
+				if (const FString* ResultTitle = Asset.TagsAndValues.Find("Title"))
+				{
+					if (ResultTitle->Contains(InText.ToString()))
+					{
+						FSearchEntry SearchEntry;
+						SearchEntry.Title = *ResultTitle;
+						SearchEntry.URL = "";
+						SearchEntry.bCategory = false;
+						SearchEntry.AssetData = Asset;
+						TutorialResults.Add(SearchEntry);
+					}
+				}
+			}
+		}
+
 		
-		TSharedRef<class IHttpRequest> HttpRequest = FHttpModule::Get().CreateRequest();
-
-		// build the url
-		// use partial response to only look at items and things we care about for items right now (title,link and label name)
-		const FText QueryURL = FText::Format( FText::FromString("https://www.googleapis.com/customsearch/v1?key=AIzaSyCMGfdDaSfjqv5zYoS0mTJnOT3e9MURWkU&cx=009868829633250020713:y7tfd8hlcgg&fields=items(title,link,labels/name)&q={0}+less:forums"), InText);
-
-		//save http request into map to ensure correct ordering
-		FText & Query = RequestQueryMap.FindOrAdd(HttpRequest);
-		Query = InText;
-
-		// kick off http request to read friends
-		HttpRequest->OnProcessRequestComplete().BindRaw(this, &SSuperSearchBox::Query_HttpRequestComplete);
-		HttpRequest->SetURL(QueryURL.ToString());
-		HttpRequest->SetHeader(TEXT("Content-Type"), TEXT("application/json"));
-		HttpRequest->SetVerb(TEXT("GET"));
-		HttpRequest->ProcessRequest();
 	}
 	else
 	{
@@ -382,18 +436,23 @@ void SSuperSearchBox::UpdateSuggestions()
 	Suggestions.Empty();
 	SelectedSuggestion = -1;
 	
-	//first show api
-	UpdateSuggestionHelper(NSLOCTEXT("SuperSearch", "API", "API"), SearchResults->OnlineResults.FindOrAdd(TEXT("api")), Suggestions);
+	//first tutorials
+	UpdateSuggestionHelper(NSLOCTEXT("SuperSearch", "tutorials", "Tutorials"), SearchResults->TutorialResults, Suggestions);
 
 	//then docs
 	UpdateSuggestionHelper(NSLOCTEXT("SuperSearch", "docs", "Documentation"), SearchResults->OnlineResults.FindOrAdd(TEXT("documentation")), Suggestions);
 
+	//then api
+	UpdateSuggestionHelper(NSLOCTEXT("SuperSearch", "API", "API"), SearchResults->OnlineResults.FindOrAdd(TEXT("api")), Suggestions);
+
 	//then wiki
 	UpdateSuggestionHelper(NSLOCTEXT("SuperSearch", "wiki", "Wiki"), SearchResults->OnlineResults.FindOrAdd(TEXT("wiki")), Suggestions);
 
 	//then answerhub
 	UpdateSuggestionHelper(NSLOCTEXT("SuperSearch", "answers", "Answerhub"), SearchResults->OnlineResults.FindOrAdd(TEXT("answers")), Suggestions);
 
+	
+
 	if(Suggestions.Num())
 	{
 		SelectedSuggestion = 1;
@@ -402,9 +461,7 @@ void SSuperSearchBox::UpdateSuggestions()
 		MarkActiveSuggestion();
 
 		// Force the textbox back into focus.
-		FWidgetPath WidgetToFocusPath;
-		FSlateApplication::Get().GeneratePathToWidgetUnchecked( InputText.ToSharedRef(), WidgetToFocusPath );
-		FSlateApplication::Get().SetKeyboardFocus( WidgetToFocusPath, EKeyboardFocusCause::SetDirectly );
+		FSlateApplication::Get().SetKeyboardFocus(InputText.ToSharedRef(), EKeyboardFocusCause::SetDirectly);
 	}
 	else
 	{
diff --git a/Engine/Source/Developer/SuperSearch/Private/SSuperSearch.h b/Engine/Source/Developer/SuperSearch/Private/SSuperSearch.h
--- a/Engine/Source/Developer/SuperSearch/Private/SSuperSearch.h
+++ b/Engine/Source/Developer/SuperSearch/Private/SSuperSearch.h
@@ -11,6 +11,7 @@ struct FSearchEntry
 	FString Title;
 	FString URL;
 	bool bCategory;
+	FAssetData AssetData;
 
 	static FSearchEntry * MakeCategoryEntry(const FString & InTitle);
 };
@@ -66,7 +67,7 @@ protected:
 
 	void Query_HttpRequestComplete(FHttpRequestPtr HttpRequest, FHttpResponsePtr HttpResponse, bool bSucceeded);
 
-	void SuggestionSelectionChanged(TSharedPtr<FSearchEntry> NewValue, ESelectInfo::Type SelectInfo);
+	void SuggestionSelectionChanged(TSharedPtr<FSearchEntry> NewValue);
 	
 	void OnMenuOpenChanged(bool bIsOpen);
 
@@ -102,6 +103,7 @@ private:
 	struct FSearchResults
 	{
 		 FCategoryResults OnlineResults;
+		 TArray<FSearchEntry> TutorialResults;
 	};
 
 	//TODO: properly clean map up as time goes by requests are processed and can be discarded
diff --git a/Engine/Source/Developer/SuperSearch/SuperSearch.Build.cs b/Engine/Source/Developer/SuperSearch/SuperSearch.Build.cs
--- a/Engine/Source/Developer/SuperSearch/SuperSearch.Build.cs
+++ b/Engine/Source/Developer/SuperSearch/SuperSearch.Build.cs
@@ -16,7 +16,8 @@ public class SuperSearch : ModuleRules
 				"Slate",
 				"SlateCore",
                 "EditorStyle",
-                "HTTP"
+                "HTTP",
+                "IntroTutorials"
 			}
 		);
 	}
