Add ask online to supersearch

diff --git a/Engine/Source/Developer/SuperSearch/Private/SSuperSearch.cpp b/Engine/Source/Developer/SuperSearch/Private/SSuperSearch.cpp
--- a/Engine/Source/Developer/SuperSearch/Private/SSuperSearch.cpp
+++ b/Engine/Source/Developer/SuperSearch/Private/SSuperSearch.cpp
@@ -7,6 +7,9 @@
 #include "IIntroTutorials.h"
 #include "EditorTutorial.h"
 
+static TSharedRef<FSearchEntry> OtherCategory(new FSearchEntry());
+static TSharedRef<FSearchEntry> AskQuestionEntry (new FSearchEntry());
+
 SSuperSearchBox::SSuperSearchBox()
 	: SelectedSuggestion(-1)
 	, bIgnoreUIUpdate(false)
@@ -16,6 +19,15 @@ SSuperSearchBox::SSuperSearchBox()
 	CategoryToIconMap.Add("Answerhub", FName("MainFrame.VisitSearchForAnswersPage") );
 	CategoryToIconMap.Add("API", FName("LevelEditor.BrowseAPIReference") );
 	CategoryToIconMap.Add("Tutorials", FName("LevelEditor.Tutorials") );
+
+
+	OtherCategory->Title = NSLOCTEXT("SuperSearchBox", "OtherCategory", "Help").ToString();
+	OtherCategory->bCategory = true;
+
+	AskQuestionEntry->Title = NSLOCTEXT("SuperSearchBox", "AskQuestion", "Ask Online").ToString();
+	AskQuestionEntry->URL = "https://answers.unrealengine.com/questions/ask.html";
+	AskQuestionEntry->bCategory = false;
+	
 }
 
 BEGIN_SLATE_FUNCTION_BUILD_OPTIMIZATION
@@ -41,7 +53,7 @@ void SSuperSearchBox::Construct( const FArguments& InArgs )
 			.Padding( FMargin(2) )
 			[
 				SNew(SBox)
-				.HeightOverride(250) // avoids flickering, ideally this would be adaptive to the content without flickering
+				//.HeightOverride(250) // avoids flickering, ideally this would be adaptive to the content without flickering
 				[
 					SAssignNew(SuggestionListView, SListView< TSharedPtr<FSearchEntry> >)
 						.ListItemsSource(&Suggestions)
@@ -451,22 +463,17 @@ void SSuperSearchBox::UpdateSuggestions()
 	//then answerhub
 	UpdateSuggestionHelper(NSLOCTEXT("SuperSearch", "answers", "Answerhub"), SearchResults->OnlineResults.FindOrAdd(TEXT("answers")), Suggestions);
 
-	
+	//finally add other category
+	Suggestions.Add(OtherCategory);
+	Suggestions.Add(AskQuestionEntry);
 
-	if(Suggestions.Num())
-	{
-		SelectedSuggestion = 1;
-		// Ideally if the selection box is open the output window is not changing it's window title (flickers)
-		SuggestionBox->SetIsOpen(true, false);
-		MarkActiveSuggestion();
+	SelectedSuggestion = 1;
+	// Ideally if the selection box is open the output window is not changing it's window title (flickers)
+	SuggestionBox->SetIsOpen(true, false);
+	MarkActiveSuggestion();
 
-		// Force the textbox back into focus.
-		FSlateApplication::Get().SetKeyboardFocus(InputText.ToSharedRef(), EKeyboardFocusCause::SetDirectly);
-	}
-	else
-	{
-		SuggestionBox->SetIsOpen(false);
-	}
+	// Force the textbox back into focus.
+	FSlateApplication::Get().SetKeyboardFocus(InputText.ToSharedRef(), EKeyboardFocusCause::SetDirectly);
 }
 
 void SSuperSearchBox::OnKeyboardFocusLost( const FKeyboardFocusEvent& InKeyboardFocusEvent )
