Fix super search to url encode its search terms.

TTP#348107

diff --git a/Engine/Source/Developer/SuperSearch/Private/SSuperSearch.cpp b/Engine/Source/Developer/SuperSearch/Private/SSuperSearch.cpp
--- a/Engine/Source/Developer/SuperSearch/Private/SSuperSearch.cpp
+++ b/Engine/Source/Developer/SuperSearch/Private/SSuperSearch.cpp
@@ -222,15 +222,18 @@ void SSuperSearchBox::OnTextChanged(const FText& InText)
 		//search through google for online content
 		{
 			TSharedRef<class IHttpRequest> HttpRequest = FHttpModule::Get().CreateRequest();
-
+			
 			// build the url
+			FString UrlEncodedString = FGenericPlatformHttp::UrlEncode(InText.ToString());	//we need to url encode for special characters (especially other languages)
 			// use partial response to only look at items and things we care about for items right now (title,link and label name)
-			const FText QueryURL = FText::Format(FText::FromString("https://www.googleapis.com/customsearch/v1?key=AIzaSyCMGfdDaSfjqv5zYoS0mTJnOT3e9MURWkU&cx=009868829633250020713:y7tfd8hlcgg&fields=items(title,link,labels/name)&q={0}+less:forums"), InText);
+			const FText QueryURL = FText::Format(FText::FromString("https://www.googleapis.com/customsearch/v1?key=AIzaSyCMGfdDaSfjqv5zYoS0mTJnOT3e9MURWkU&cx=009868829633250020713:y7tfd8hlcgg&fields=items(title,link,labels/name)&q={0}+less:forums"), FText::FromString(UrlEncodedString));
 
 			//save http request into map to ensure correct ordering
 			FText & Query = RequestQueryMap.FindOrAdd(HttpRequest);
 			Query = InText;
 
+			
+
 			// kick off http request to read friends
 			HttpRequest->OnProcessRequestComplete().BindRaw(this, &SSuperSearchBox::Query_HttpRequestComplete);
 			HttpRequest->SetURL(QueryURL.ToString());
