Made the Super Search appear as an overlay rather than be its own window

This works around an issue where the active IME context is deactivated due to the HWND changing.

diff --git a/Engine/Source/Developer/SuperSearch/Private/SSuperSearch.cpp b/Engine/Source/Developer/SuperSearch/Private/SSuperSearch.cpp
--- a/Engine/Source/Developer/SuperSearch/Private/SSuperSearch.cpp
+++ b/Engine/Source/Developer/SuperSearch/Private/SSuperSearch.cpp
@@ -37,6 +37,7 @@ void SSuperSearchBox::Construct( const FArguments& InArgs )
 	[
 		SAssignNew( SuggestionBox, SMenuAnchor )
 		.Placement( InArgs._SuggestionListPlacement )
+		.Method( SMenuAnchor::UseCurrentWindow )
 		[
 			SAssignNew(InputText, SSearchBox)
 			.OnTextCommitted(this, &SSuperSearchBox::OnTextCommitted)
