Fix for the Super Search suggestions list sometimes being cut off.

The Super Search suggestions list now anchors itself to the bottom right (rather than the bottom left) of the text box that spawns it - this should prevent it from extending beyond the bounds of the window that owns it.

diff --git a/Engine/Source/Developer/SuperSearch/Private/SSuperSearch.h b/Engine/Source/Developer/SuperSearch/Private/SSuperSearch.h
--- a/Engine/Source/Developer/SuperSearch/Private/SSuperSearch.h
+++ b/Engine/Source/Developer/SuperSearch/Private/SSuperSearch.h
@@ -24,7 +24,7 @@ class SSuperSearchBox
 public:
 
 	SLATE_BEGIN_ARGS( SSuperSearchBox )
-		: _SuggestionListPlacement( MenuPlacement_BelowAnchor )
+		: _SuggestionListPlacement( MenuPlacement_ComboBoxRight )
 		{}
 
 		/** Where to place the suggestion list */
