[INTEGRATE] UE4-LauncherReleases->UE4 integrate all from launcher releases up to 2353416, including supersearch bar integration, free purchase flow, and various bug fixes

diff --git a/Engine/Source/Developer/SuperSearch/Private/SSuperSearch.cpp b/Engine/Source/Developer/SuperSearch/Private/SSuperSearch.cpp
--- a/Engine/Source/Developer/SuperSearch/Private/SSuperSearch.cpp
+++ b/Engine/Source/Developer/SuperSearch/Private/SSuperSearch.cpp
@@ -1,13 +1,21 @@
 // Copyright 1998-2014 Epic Games, Inc. All Rights Reserved.
 
 #include "SuperSearchPrivatePCH.h"
+#include "SSearchBox.h"
 #include "SSuperSearch.h"
 #include "SScrollBorder.h"
+
+#if WITH_EDITOR
 #include "AssetRegistryModule.h"
 #include "IIntroTutorials.h"
 #include "EditorTutorial.h"
+#endif
+
 #include "SSearchBox.h"
+
+#if WITH_EDITOR	
 #include "AssetData.h"
+#endif
 
 static TSharedRef<FSearchEntry> OtherCategory(new FSearchEntry());
 static TSharedRef<FSearchEntry> AskQuestionEntry (new FSearchEntry());
@@ -38,36 +46,36 @@ void SSuperSearchBox::Construct( const FArguments& InArgs )
 	ChildSlot
 	[
 		SAssignNew( SuggestionBox, SMenuAnchor )
-		.Placement( InArgs._SuggestionListPlacement )
+			.Placement( InArgs._SuggestionListPlacement )
 		.Method( SMenuAnchor::UseCurrentWindow )
-		[
-			SAssignNew(InputText, SSearchBox)
-			.OnTextCommitted(this, &SSuperSearchBox::OnTextCommitted)
-			.HintText(NSLOCTEXT("SuperSearchBox", "HelpHint", "Search For Help"))
-			.OnTextChanged(this, &SSuperSearchBox::OnTextChanged)
-			.SelectAllTextWhenFocused(false)
-			.DelayChangeNotificationsWhileTyping(true)
-			.MinDesiredWidth(200)
-		]
-		.MenuContent
-		(
-			SNew(SBorder)
-			.BorderImage(FEditorStyle::GetBrush("Menu.Background"))
-			.Padding( FMargin(2) )
 			[
-				SNew(SBox)
-				//.HeightOverride(250) // avoids flickering, ideally this would be adaptive to the content without flickering
+				SAssignNew(InputText, SSearchBox)
+					.OnTextCommitted(this, &SSuperSearchBox::OnTextCommitted)
+					.HintText( NSLOCTEXT( "SuperSearchBox", "HelpHint", "Search For Help" ) )
+					.OnTextChanged(this, &SSuperSearchBox::OnTextChanged)
+					.SelectAllTextWhenFocused(false)
+					.DelayChangeNotificationsWhileTyping(true)
+			.MinDesiredWidth(200)
+			]
+			.MenuContent
+			(
+				SNew(SBorder)
+				.BorderImage(FCoreStyle::Get().GetBrush("Menu.Background"))
+				.Padding( FMargin(2) )
 				[
-					SAssignNew(SuggestionListView, SListView< TSharedPtr<FSearchEntry> >)
-						.ListItemsSource(&Suggestions)
-						.SelectionMode( ESelectionMode::Single )							// Ideally the mouse over would not highlight while keyboard controls the UI
-						.OnGenerateRow(this, &SSuperSearchBox::MakeSuggestionListItemWidget)
+					SNew(SBox)
+				//.HeightOverride(250) // avoids flickering, ideally this would be adaptive to the content without flickering
+					[
+						SAssignNew(SuggestionListView, SListView< TSharedPtr<FSearchEntry> >)
+							.ListItemsSource(&Suggestions)
+							.SelectionMode( ESelectionMode::Single )							// Ideally the mouse over would not highlight while keyboard controls the UI
+							.OnGenerateRow(this, &SSuperSearchBox::MakeSuggestionListItemWidget)
 						.OnMouseButtonClick(this, &SSuperSearchBox::SuggestionSelectionChanged)
-						.ItemHeight(18)
+							.ItemHeight(18)
+					]
 				]
-			]
-		)
-		.OnMenuOpenChanged(this, &SSuperSearchBox::OnMenuOpenChanged)
+			)
+			.OnMenuOpenChanged(this, &SSuperSearchBox::OnMenuOpenChanged)
 	];
 }
 END_SLATE_FUNCTION_BUILD_OPTIMIZATION
@@ -78,9 +86,9 @@ int32 GetNumRealSuggestions(const TArray< TSharedPtr<FSearchEntry> > & Suggestio
 	for (auto Entry : Suggestions)
 	{
 		if (Entry->bCategory == false)
-		{
+	{
 			++NumSuggestions;
-		}
+	}
 	}
 
 	return NumSuggestions;
@@ -88,6 +96,7 @@ int32 GetNumRealSuggestions(const TArray< TSharedPtr<FSearchEntry> > & Suggestio
 
 void SSuperSearchBox::OnMenuOpenChanged( bool bIsOpen )
 {
+#if WITH_EDITOR
 	if (bIsOpen == false)
 	{
 		if(FEngineAnalytics::IsAvailable())
@@ -103,7 +112,7 @@ void SSuperSearchBox::OnMenuOpenChanged( bool bIsOpen )
 
 		InputText->SetText(FText::GetEmpty());
 	}
-
+#endif
 	EntryClicked = NULL;
 }
 
@@ -112,6 +121,7 @@ void SSuperSearchBox::ActOnSuggestion(TSharedPtr<FSearchEntry> SearchEntry)
 	if (SearchEntry->bCategory == false)
 	{
 		EntryClicked = SearchEntry;
+#if WITH_EDITOR
 		if (SearchEntry->URL.IsEmpty())
 		{
 			UBlueprint* Blueprint = LoadObject<UBlueprint>(nullptr, *(SearchEntry->AssetData.ObjectPath.ToString()));
@@ -124,6 +134,7 @@ void SSuperSearchBox::ActOnSuggestion(TSharedPtr<FSearchEntry> SearchEntry)
 			IntroTutorials.LaunchTutorial(Tutorial, true, ParentWindow);
 		}
 		else
+#endif
 		{
 			FPlatformProcess::LaunchURL(*SearchEntry->URL, NULL, NULL);
 		}
@@ -151,7 +162,7 @@ void SSuperSearchBox::SuggestionSelectionChanged(TSharedPtr<FSearchEntry> NewVal
 			SelectedSuggestion = i;
 			MarkActiveSuggestion();
 
-			ActOnSuggestion(NewValue);
+				ActOnSuggestion(NewValue);
 
 			break;
 		}
@@ -165,24 +176,31 @@ TSharedRef<ITableRow> SSuperSearchBox::MakeSuggestionListItemWidget(TSharedPtr<F
 
 	FString Left, Right, Combined;
 
+	if (Suggestion->bCategory)
+	{
+		Combined = "- ";
+	}
+
 	if(Suggestion->Title.Split(TEXT("\t"), &Left, &Right))
 	{
-		Combined = Left + Right;
+		Combined += Left + Right;
 	}
 	else
 	{
-		Combined = Suggestion->Title;
+		Combined += Suggestion->Title;
 	}
 
 	bool bCategory = Suggestion->bCategory;
 
-	
+
 	TSharedPtr<SWidget> IconWidget;
 	if (bCategory)
 	{
+#if WITH_EDITOR
 		FName* IconResult = CategoryToIconMap.Find(Combined);
 		SAssignNew(IconWidget, SImage)
 			.Image(FEditorStyle::GetBrush(IconResult ? *IconResult : FName("MainFrame.VisitEpicGamesDotCom") ));
+#endif
 	}
 	else
 	{
@@ -195,17 +213,21 @@ TSharedRef<ITableRow> SSuperSearchBox::MakeSuggestionListItemWidget(TSharedPtr<F
 		.Padding(FMargin(bCategory ? 2.5 : 16, 2.5, 2.5, 2.5))
 		[
 			SNew(SHorizontalBox)
+#if WITH_EDITOR
 			+ SHorizontalBox::Slot()
 			.Padding(FMargin(0,0,2.5,0))
 			[
 				IconWidget.ToSharedRef()
 			]
+#endif
 			+SHorizontalBox::Slot()
 			.AutoWidth()
 			[
 				SNew(STextBlock)
 				.Text(Combined)
+#if WITH_EDITOR
 				.TextStyle(&FEditorStyle::Get().GetWidgetStyle<FTextBlockStyle>(bCategory ? "SuperSearchCategoryText" : "NormalText"))
+#endif
 				.MinDesiredWidth(600)
 			]
 			
@@ -225,14 +247,14 @@ void SSuperSearchBox::OnTextChanged(const FText& InText)
 		//search through google for online content
 		{
 			TSharedRef<class IHttpRequest> HttpRequest = FHttpModule::Get().CreateRequest();
-			
+
 			// build the url
-			FString UrlEncodedString = FGenericPlatformHttp::UrlEncode(InText.ToString());	//we need to url encode for special characters (especially other languages)
+				FString UrlEncodedString = FGenericPlatformHttp::UrlEncode(InText.ToString());	//we need to url encode for special characters (especially other languages)
 			// use partial response to only look at items and things we care about for items right now (title,link and label name)
-			const FText QueryURL = FText::Format(FText::FromString("https://www.googleapis.com/customsearch/v1?key=AIzaSyCMGfdDaSfjqv5zYoS0mTJnOT3e9MURWkU&cx=009868829633250020713:y7tfd8hlcgg&fields=items(title,link,labels/name)&q={0}+less:forums"), FText::FromString(UrlEncodedString));
+				const FText QueryURL = FText::Format(FText::FromString("https://www.googleapis.com/customsearch/v1?key=AIzaSyCMGfdDaSfjqv5zYoS0mTJnOT3e9MURWkU&cx=009868829633250020713:y7tfd8hlcgg&fields=items(title,link,labels/name)&q={0}+less:forums"), FText::FromString(UrlEncodedString));
 
 			//save http request into map to ensure correct ordering
-			FText& Query = RequestQueryMap.FindOrAdd(HttpRequest);
+			FText & Query = RequestQueryMap.FindOrAdd(HttpRequest);
 			Query = InText;
 
 			
@@ -245,6 +267,7 @@ void SSuperSearchBox::OnTextChanged(const FText& InText)
 			HttpRequest->ProcessRequest();
 		}
 
+#if WITH_EDITOR
 		//search local tutorials
 		{
 			FAssetRegistryModule& AssetRegistry = FModuleManager::LoadModuleChecked<FAssetRegistryModule>(TEXT("AssetRegistry"));
@@ -278,8 +301,7 @@ void SSuperSearchBox::OnTextChanged(const FText& InText)
 				}
 			}
 		}
-
-		
+#endif
 	}
 	else
 	{
@@ -311,7 +333,7 @@ void SSuperSearchBox::Query_HttpRequestComplete(FHttpRequestPtr HttpRequest, FHt
 
 			if (FJsonSerializer::Deserialize(JsonReader, JsonObject) && JsonObject.IsValid())
 			{
-				if (FText* QueryText = RequestQueryMap.Find(HttpRequest))
+				if (FText * QueryText = RequestQueryMap.Find(HttpRequest))
 				{
 					const TArray<TSharedPtr<FJsonValue> > * JsonItems = nullptr;
 					if (JsonObject->TryGetArrayField(TEXT("items"), JsonItems))
@@ -417,7 +439,7 @@ FReply SSuperSearchBox::OnKeyDown( const FGeometry& MyGeometry, const FKeyEvent&
 	return FReply::Unhandled();
 }
 
-FSearchEntry * FSearchEntry::MakeCategoryEntry(const FString& InTitle)
+FSearchEntry * FSearchEntry::MakeCategoryEntry(const FString & InTitle)
 {
 	FSearchEntry * SearchEntry = new FSearchEntry();
 	SearchEntry->Title = InTitle;
@@ -426,7 +448,7 @@ FSearchEntry * FSearchEntry::MakeCategoryEntry(const FString& InTitle)
 	return SearchEntry;
 }
 
-void UpdateSuggestionHelper(const FText& CategoryLabel, const TArray<FSearchEntry> & Elements, TArray<TSharedPtr< FSearchEntry > > & OutSuggestions)
+void UpdateSuggestionHelper(const FText & CategoryLabel, const TArray<FSearchEntry> & Elements, TArray<TSharedPtr< FSearchEntry > > & OutSuggestions)
 {
 	if (Elements.Num())
 	{
@@ -441,7 +463,7 @@ void UpdateSuggestionHelper(const FText& CategoryLabel, const TArray<FSearchEntr
 
 void SSuperSearchBox::UpdateSuggestions()
 {
-	const FText& Query = InputText->GetText();
+	const FText & Query = InputText->GetText();
 	FSearchResults * SearchResults = SearchResultsCache.Find(Query.ToString());
 
 	//still waiting on results for current query
@@ -473,12 +495,12 @@ void SSuperSearchBox::UpdateSuggestions()
 	Suggestions.Add(OtherCategory);
 	Suggestions.Add(AskQuestionEntry);
 
-	SelectedSuggestion = 1;
-	// Ideally if the selection box is open the output window is not changing it's window title (flickers)
-	SuggestionBox->SetIsOpen(true, false);
-	MarkActiveSuggestion();
+		SelectedSuggestion = 1;
+		// Ideally if the selection box is open the output window is not changing it's window title (flickers)
+		SuggestionBox->SetIsOpen(true, false);
+		MarkActiveSuggestion();
 
-	// Force the textbox back into focus.
+		// Force the textbox back into focus.
 	FSlateApplication::Get().SetKeyboardFocus(InputText.ToSharedRef(), EFocusCause::SetDirectly);
 }
 
diff --git a/Engine/Source/Developer/SuperSearch/Private/SSuperSearch.h b/Engine/Source/Developer/SuperSearch/Private/SSuperSearch.h
--- a/Engine/Source/Developer/SuperSearch/Private/SSuperSearch.h
+++ b/Engine/Source/Developer/SuperSearch/Private/SSuperSearch.h
@@ -1,7 +1,10 @@
 // Copyright 1998-2014 Epic Games, Inc. All Rights Reserved.
 
 #pragma once
+
+#if WITH_EDITOR
 #include "AssetData.h"
+#endif
 
 /**
  * Console input box with command-completion support
@@ -12,8 +15,9 @@ struct FSearchEntry
 	FString Title;
 	FString URL;
 	bool bCategory;
+#if WITH_EDITOR
 	FAssetData AssetData;
-
+#endif
 	static FSearchEntry * MakeCategoryEntry(const FString& InTitle);
 };
 
diff --git a/Engine/Source/Developer/SuperSearch/Private/SuperSearchPrivatePCH.h b/Engine/Source/Developer/SuperSearch/Private/SuperSearchPrivatePCH.h
--- a/Engine/Source/Developer/SuperSearch/Private/SuperSearchPrivatePCH.h
+++ b/Engine/Source/Developer/SuperSearch/Private/SuperSearchPrivatePCH.h
@@ -2,7 +2,6 @@
 
 #pragma once
 
-#include "UnrealEd.h"
 #include "SuperSearchModule.h"
 #include "SlateBasics.h"
 #include "Http.h"
diff --git a/Engine/Source/Developer/SuperSearch/SuperSearch.Build.cs b/Engine/Source/Developer/SuperSearch/SuperSearch.Build.cs
--- a/Engine/Source/Developer/SuperSearch/SuperSearch.Build.cs
+++ b/Engine/Source/Developer/SuperSearch/SuperSearch.Build.cs
@@ -6,20 +6,35 @@ public class SuperSearch : ModuleRules
 {
 	public SuperSearch(TargetInfo Target)
 	{
-		PrivateDependencyModuleNames.AddRange(
-			new string[] {
-				"Core",
-				"CoreUObject", // @todo Mac: for some reason it's needed to link in debug on Mac
-				"Engine",
-                "InputCore",
-				"Json",
-				"UnrealEd",
-				"Slate",
-				"SlateCore",
-                "EditorStyle",
-                "HTTP",
-                "IntroTutorials"
-			}
-		);
+		if (UEBuildConfiguration.bBuildEditor == true)
+		{
+			PrivateDependencyModuleNames.AddRange(
+				new string[] {
+					"Core",
+					"CoreUObject", // @todo Mac: for some reason it's needed to link in debug on Mac
+					"Engine",
+					"InputCore",
+					"Json",
+					"UnrealEd",
+					"Slate",
+					"SlateCore",
+					"EditorStyle",
+					"HTTP",
+					"IntroTutorials"
+				}
+			);
+		}
+		else
+		{
+	   		PrivateDependencyModuleNames.AddRange(
+				new string[] {
+					"Core",
+					"CoreUObject", // @todo Mac: for some reason it's needed to link in debug on Mac
+					"Slate",
+					"SlateCore",
+					"HTTP"
+				}
+			);
+		}
 	}
 }
diff --git a/Engine/Source/Editor/EditorStyle/Private/SlateEditorStyle.cpp b/Engine/Source/Editor/EditorStyle/Private/SlateEditorStyle.cpp
--- a/Engine/Source/Editor/EditorStyle/Private/SlateEditorStyle.cpp
+++ b/Engine/Source/Editor/EditorStyle/Private/SlateEditorStyle.cpp
@@ -2288,44 +2288,44 @@ void FSlateEditorStyle::FStyle::SetupWindowStyles()
 		else
 		{
 			MinimizeButtonStyle = FButtonStyle(Button)
-				.SetNormal ( IMAGE_BRUSH( "Old/Window/WindowButton_Mac_Yosemite_Minimize_Normal", FVector2D(20, 20) ) )
-				.SetHovered( IMAGE_BRUSH( "Old/Window/WindowButton_Mac_Yosemite_Minimize_Hovered", FVector2D(20, 20) ) )
-				.SetPressed( IMAGE_BRUSH( "Old/Window/WindowButton_Mac_Yosemite_Minimize_Pressed", FVector2D(20, 20) ) )
-				.SetDisabled( IMAGE_BRUSH( "Old/Window/WindowButton_Mac_Yosemite_Disabled", FVector2D(20, 20) ) );
+				.SetNormal ( IMAGE_BRUSH( "Old/Window/Mac_Yosemite_Minimize_Normal", FVector2D(20, 20) ) )
+				.SetHovered( IMAGE_BRUSH( "Old/Window/Mac_Yosemite_Minimize_Hovered", FVector2D(20, 20) ) )
+				.SetPressed( IMAGE_BRUSH( "Old/Window/Mac_Yosemite_Minimize_Pressed", FVector2D(20, 20) ) )
+				.SetDisabled( IMAGE_BRUSH( "Old/Window/Mac_Yosemite_Disabled", FVector2D(20, 20) ) );
 
-			Set( "Window.Buttons.Minimize.Normal",    new IMAGE_BRUSH( "Old/Window/WindowButton_Mac_Yosemite_Minimize_Normal", FVector2D(20, 20) ) );
-			Set( "Window.Buttons.Minimize.Hovered",   new IMAGE_BRUSH( "Old/Window/WindowButton_Mac_Yosemite_Minimize_Hovered", FVector2D(20, 20) ) );
-			Set( "Window.Buttons.Minimize.Pressed",   new IMAGE_BRUSH( "Old/Window/WindowButton_Mac_Yosemite_Minimize_Pressed", FVector2D(20, 20) ) );
-			Set( "Window.Buttons.Minimize.Disabled",  new IMAGE_BRUSH( "Old/Window/WindowButton_Mac_Yosemite_Disabled", FVector2D(20, 20) ) );
+			Set( "Window.Buttons.Minimize.Normal",    new IMAGE_BRUSH( "Old/Window/Mac_Yosemite_Minimize_Normal", FVector2D(20, 20) ) );
+			Set( "Window.Buttons.Minimize.Hovered",   new IMAGE_BRUSH( "Old/Window/Mac_Yosemite_Minimize_Hovered", FVector2D(20, 20) ) );
+			Set( "Window.Buttons.Minimize.Pressed",   new IMAGE_BRUSH( "Old/Window/Mac_Yosemite_Minimize_Pressed", FVector2D(20, 20) ) );
+			Set( "Window.Buttons.Minimize.Disabled",  new IMAGE_BRUSH( "Old/Window/Mac_Yosemite_Disabled", FVector2D(20, 20) ) );
 
 			MaximizeButtonStyle = FButtonStyle(Button)
-				.SetNormal ( IMAGE_BRUSH( "Old/Window/WindowButton_Mac_Yosemite_Maximize_Normal", FVector2D(20, 20) ) )
-				.SetHovered( IMAGE_BRUSH( "Old/Window/WindowButton_Mac_Yosemite_Maximize_Hovered", FVector2D(20, 20) ) )
-				.SetPressed( IMAGE_BRUSH( "Old/Window/WindowButton_Mac_Yosemite_Maximize_Pressed", FVector2D(20, 20) ) )
-				.SetDisabled( IMAGE_BRUSH( "Old/Window/WindowButton_Mac_Yosemite_Disabled", FVector2D(20, 20) ) );
+				.SetNormal ( IMAGE_BRUSH( "Old/Window/Mac_Yosemite_Maximize_Normal", FVector2D(20, 20) ) )
+				.SetHovered( IMAGE_BRUSH( "Old/Window/Mac_Yosemite_Maximize_Hovered", FVector2D(20, 20) ) )
+				.SetPressed( IMAGE_BRUSH( "Old/Window/Mac_Yosemite_Maximize_Pressed", FVector2D(20, 20) ) )
+				.SetDisabled( IMAGE_BRUSH( "Old/Window/Mac_Yosemite_Disabled", FVector2D(20, 20) ) );
 
-			Set( "Window.Buttons.Maximize.Normal",    new IMAGE_BRUSH( "Old/Window/WindowButton_Mac_Yosemite_Maximize_Normal", FVector2D(20, 20) ) );
-			Set( "Window.Buttons.Maximize.Hovered",   new IMAGE_BRUSH( "Old/Window/WindowButton_Mac_Yosemite_Maximize_Hovered", FVector2D(20, 20) ) );
-			Set( "Window.Buttons.Maximize.Pressed",   new IMAGE_BRUSH( "Old/Window/WindowButton_Mac_Yosemite_Maximize_Pressed", FVector2D(20, 20) ) );
-			Set( "Window.Buttons.Maximize.Disabled",  new IMAGE_BRUSH( "Old/Window/WindowButton_Mac_Yosemite_Disabled", FVector2D(20, 20) ) );
+			Set( "Window.Buttons.Maximize.Normal",    new IMAGE_BRUSH( "Old/Window/Mac_Yosemite_Maximize_Normal", FVector2D(20, 20) ) );
+			Set( "Window.Buttons.Maximize.Hovered",   new IMAGE_BRUSH( "Old/Window/Mac_Yosemite_Maximize_Hovered", FVector2D(20, 20) ) );
+			Set( "Window.Buttons.Maximize.Pressed",   new IMAGE_BRUSH( "Old/Window/Mac_Yosemite_Maximize_Pressed", FVector2D(20, 20) ) );
+			Set( "Window.Buttons.Maximize.Disabled",  new IMAGE_BRUSH( "Old/Window/Mac_Yosemite_Disabled", FVector2D(20, 20) ) );
 
 			RestoreButtonStyle = FButtonStyle(Button)
-				.SetNormal ( IMAGE_BRUSH( "Old/Window/WindowButton_Mac_Yosemite_Maximize_Normal", FVector2D(20, 20) ) )
-				.SetHovered( IMAGE_BRUSH( "Old/Window/WindowButton_Mac_Yosemite_Maximize_Hovered", FVector2D(20, 20) ) )
-				.SetPressed( IMAGE_BRUSH( "Old/Window/WindowButton_Mac_Yosemite_Maximize_Pressed", FVector2D(20, 20) ) );
+				.SetNormal ( IMAGE_BRUSH( "Old/Window/Mac_Yosemite_Maximize_Normal", FVector2D(20, 20) ) )
+				.SetHovered( IMAGE_BRUSH( "Old/Window/Mac_Yosemite_Maximize_Hovered", FVector2D(20, 20) ) )
+				.SetPressed( IMAGE_BRUSH( "Old/Window/Mac_Yosemite_Maximize_Pressed", FVector2D(20, 20) ) );
 
-			Set( "Window.Buttons.Restore.Normal",     new IMAGE_BRUSH( "Old/Window/WindowButton_Mac_Yosemite_Maximize_Normal", FVector2D(20, 20) ) );
-			Set( "Window.Buttons.Restore.Hovered",    new IMAGE_BRUSH( "Old/Window/WindowButton_Mac_Yosemite_Maximize_Hovered", FVector2D(20, 20) ) );
-			Set( "Window.Buttons.Restore.Pressed",    new IMAGE_BRUSH( "Old/Window/WindowButton_Mac_Yosemite_Maximize_Pressed", FVector2D(20, 20) ) );
+			Set( "Window.Buttons.Restore.Normal",     new IMAGE_BRUSH( "Old/Window/Mac_Yosemite_Maximize_Normal", FVector2D(20, 20) ) );
+			Set( "Window.Buttons.Restore.Hovered",    new IMAGE_BRUSH( "Old/Window/Mac_Yosemite_Maximize_Hovered", FVector2D(20, 20) ) );
+			Set( "Window.Buttons.Restore.Pressed",    new IMAGE_BRUSH( "Old/Window/Mac_Yosemite_Maximize_Pressed", FVector2D(20, 20) ) );
 
 			CloseButtonStyle = FButtonStyle(Button)
-				.SetNormal ( IMAGE_BRUSH( "Old/Window/WindowButton_Mac_Yosemite_Close_Normal", FVector2D(20, 20) ) )
-				.SetHovered( IMAGE_BRUSH( "Old/Window/WindowButton_Mac_Yosemite_Close_Hovered", FVector2D(20, 20) ) )
-				.SetPressed( IMAGE_BRUSH( "Old/Window/WindowButton_Mac_Yosemite_Close_Pressed", FVector2D(20, 20) ) );
+				.SetNormal ( IMAGE_BRUSH( "Old/Window/Mac_Yosemite_Close_Normal", FVector2D(20, 20) ) )
+				.SetHovered( IMAGE_BRUSH( "Old/Window/Mac_Yosemite_Close_Hovered", FVector2D(20, 20) ) )
+				.SetPressed( IMAGE_BRUSH( "Old/Window/Mac_Yosemite_Close_Pressed", FVector2D(20, 20) ) );
 
-			Set( "Window.Buttons.Close.Normal",       new IMAGE_BRUSH( "Old/Window/WindowButton_Mac_Yosemite_Close_Normal", FVector2D(20, 20) ) );
-			Set( "Window.Buttons.Close.Hovered",      new IMAGE_BRUSH( "Old/Window/WindowButton_Mac_Yosemite_Close_Hovered", FVector2D(20, 20) ) );
-			Set( "Window.Buttons.Close.Pressed",      new IMAGE_BRUSH( "Old/Window/WindowButton_Mac_Yosemite_Close_Pressed", FVector2D(20, 20) ) );
+			Set( "Window.Buttons.Close.Normal",       new IMAGE_BRUSH( "Old/Window/Mac_Yosemite_Close_Normal", FVector2D(20, 20) ) );
+			Set( "Window.Buttons.Close.Hovered",      new IMAGE_BRUSH( "Old/Window/Mac_Yosemite_Close_Hovered", FVector2D(20, 20) ) );
+			Set( "Window.Buttons.Close.Pressed",      new IMAGE_BRUSH( "Old/Window/Mac_Yosemite_Close_Pressed", FVector2D(20, 20) ) );
 		}
 #else
 		const FButtonStyle MinimizeButtonStyle = FButtonStyle(Button)
