Re-add "UnrealEd.h" to make sure proper headers are iincluded

diff --git a/Engine/Source/Developer/SuperSearch/Private/SuperSearchPrivatePCH.h b/Engine/Source/Developer/SuperSearch/Private/SuperSearchPrivatePCH.h
--- a/Engine/Source/Developer/SuperSearch/Private/SuperSearchPrivatePCH.h
+++ b/Engine/Source/Developer/SuperSearch/Private/SuperSearchPrivatePCH.h
@@ -2,6 +2,9 @@
 
 #pragma once
 
+#if WITH_EDITOR
+#include "UnrealEd.h"
+#endif
 #include "SuperSearchModule.h"
 #include "SlateBasics.h"
 #include "Http.h"
