Search For Help box in the editor now also appends source data to documentation URLs.

diff --git a/Engine/Source/Developer/SuperSearch/Private/SSuperSearch.cpp b/Engine/Source/Developer/SuperSearch/Private/SSuperSearch.cpp
--- a/Engine/Source/Developer/SuperSearch/Private/SSuperSearch.cpp
+++ b/Engine/Source/Developer/SuperSearch/Private/SSuperSearch.cpp
@@ -120,7 +120,42 @@ void SSuperSearchBox::OnMenuOpenChanged( bool bIsOpen )
 	EntryClicked = NULL;
 }
 
-void SSuperSearchBox::ActOnSuggestion(TSharedPtr<FSearchEntry> SearchEntry)
+/** Input string should have no leading & or ? */
+static FString AddQueryStringToURL(FString const& URL, TCHAR const* StringToAdd)
+{
+	if (StringToAdd)
+	{
+		// split url into path, querystring, and anchor parts
+		FString Anchor;
+		FString Path;
+		FString QueryString;
+		FString PathAndQueryString;
+		if (URL.Split(TEXT("#"), &PathAndQueryString, &Anchor) == false)
+		{
+			PathAndQueryString = URL;
+		}
+		if (PathAndQueryString.Split(TEXT("?"), &Path, &QueryString) == false)
+		{
+			Path = URL;
+		}
+
+		if (QueryString.IsEmpty())
+		{
+			QueryString = FString::Printf(TEXT("?%s"), StringToAdd);
+		}
+		else
+		{
+			QueryString = FString::Printf(TEXT("%s&%s"), *QueryString, StringToAdd);
+		}
+
+		// reassemble and return
+		return Path + QueryString + Anchor;
+	}
+
+	return URL;
+}
+
+void SSuperSearchBox::ActOnSuggestion(TSharedPtr<FSearchEntry> SearchEntry, FString const& Category)
 {
 	if (SearchEntry->bCategory == false)
 	{
@@ -140,7 +175,16 @@ void SSuperSearchBox::ActOnSuggestion(TSharedPtr<FSearchEntry> SearchEntry)
 		else
 #endif
 		{
-			FPlatformProcess::LaunchURL(*SearchEntry->URL, NULL, NULL);
+			if (Category == TEXT("documentation"))
+			{
+				// append some tracking data to the URL
+				FString const FinalURL = AddQueryStringToURL(SearchEntry->URL, TEXT("utm_source=editor&utm_medium=search&utm_campaign=user_initiated"));
+				FPlatformProcess::LaunchURL(*FinalURL, NULL, NULL);
+			}
+			else
+			{
+				FPlatformProcess::LaunchURL(*SearchEntry->URL, NULL, NULL);
+			}
 		}
 
 		SuggestionBox->SetIsOpen(false);
@@ -159,14 +203,21 @@ void SSuperSearchBox::SuggestionSelectionChanged(TSharedPtr<FSearchEntry> NewVal
 		return;
 	}
 
+	int32 CurrentCategory = INDEX_NONE;
 	for(int32 i = 0; i < Suggestions.Num(); ++i)
 	{
+		if (Suggestions[i]->bCategory)
+		{
+			CurrentCategory = i;
+		}
+
 		if(NewValue == Suggestions[i])
 		{
 			SelectedSuggestion = i;
 			MarkActiveSuggestion();
 
-				ActOnSuggestion(NewValue);
+			FString const Category = CurrentCategory != INDEX_NONE ? Suggestions[CurrentCategory]->Title : FString();
+			ActOnSuggestion(NewValue, Category);
 
 			break;
 		}
@@ -317,7 +368,18 @@ void SSuperSearchBox::OnTextCommitted( const FText& InText, ETextCommit::Type Co
 {
 	if (CommitInfo == ETextCommit::OnEnter && SelectedSuggestion >= 1)
 	{
-		ActOnSuggestion(Suggestions[SelectedSuggestion]);
+		// find the category
+		FString Category;
+		for (int32 i = SelectedSuggestion; i >= 0; --i)
+		{
+			if (Suggestions[i]->bCategory)
+			{
+				Category = Suggestions[i]->Title;
+				break;
+			}
+		}
+
+		ActOnSuggestion(Suggestions[SelectedSuggestion], Category);
 	}
 }
 
diff --git a/Engine/Source/Developer/SuperSearch/Private/SSuperSearch.h b/Engine/Source/Developer/SuperSearch/Private/SSuperSearch.h
--- a/Engine/Source/Developer/SuperSearch/Private/SSuperSearch.h
+++ b/Engine/Source/Developer/SuperSearch/Private/SSuperSearch.h
@@ -80,7 +80,7 @@ protected:
 	
 	void OnMenuOpenChanged(bool bIsOpen);
 
-	void ActOnSuggestion(TSharedPtr<FSearchEntry> SearchEntry);
+	void ActOnSuggestion(TSharedPtr<FSearchEntry> SearchEntry, FString const& Category);
 		
 	void UpdateSuggestions();
 
