Fix search results not showing up if user types quickly and hits Enter

diff --git a/Engine/Source/Developer/SuperSearch/Private/SSuperSearch.cpp b/Engine/Source/Developer/SuperSearch/Private/SSuperSearch.cpp
--- a/Engine/Source/Developer/SuperSearch/Private/SSuperSearch.cpp
+++ b/Engine/Source/Developer/SuperSearch/Private/SSuperSearch.cpp
@@ -366,20 +366,28 @@ void SSuperSearchBox::OnTextChanged(const FText& InText)
 
 void SSuperSearchBox::OnTextCommitted( const FText& InText, ETextCommit::Type CommitInfo)
 {
-	if (CommitInfo == ETextCommit::OnEnter && SelectedSuggestion >= 1)
+	if (CommitInfo == ETextCommit::OnEnter)
 	{
-		// find the category
-		FString Category;
-		for (int32 i = SelectedSuggestion; i >= 0; --i)
+		if (SelectedSuggestion >= 1)
 		{
-			if (Suggestions[i]->bCategory)
+			// find the category
+			FString Category;
+			for (int32 i = SelectedSuggestion; i >= 0; --i)
 			{
-				Category = Suggestions[i]->Title;
-				break;
+				if (Suggestions[i]->bCategory)
+				{
+					Category = Suggestions[i]->Title;
+					break;
+				}
 			}
+
+			ActOnSuggestion(Suggestions[SelectedSuggestion], Category);
+		}
+		else
+		{
+			OnTextChanged(InText);
 		}
 
-		ActOnSuggestion(Suggestions[SelectedSuggestion], Category);
 	}
 }
 
