Added tags to tutorials
Made super search search tutorial tags

diff --git a/Engine/Source/Developer/SuperSearch/Private/SSuperSearch.cpp b/Engine/Source/Developer/SuperSearch/Private/SSuperSearch.cpp
--- a/Engine/Source/Developer/SuperSearch/Private/SSuperSearch.cpp
+++ b/Engine/Source/Developer/SuperSearch/Private/SSuperSearch.cpp
@@ -342,7 +342,9 @@ void SSuperSearchBox::OnTextChanged(const FText& InText)
 
 			for (const FAssetData& Asset : AssetData)
 			{
-				if (const FString* ResultTitle = Asset.TagsAndValues.Find("Title"))
+				const FString* SearchTag = Asset.TagsAndValues.Find("SearchTags");
+				const FString* ResultTitle = Asset.TagsAndValues.Find("Title");
+				if( ResultTitle)
 				{
 					if (ResultTitle->Contains(InText.ToString()))
 					{
@@ -354,6 +356,28 @@ void SSuperSearchBox::OnTextChanged(const FText& InText)
 						TutorialResults.Add(SearchEntry);
 					}
 				}
+
+				// If the asset has search tags, search them
+				if ((SearchTag) && (SearchTag->IsEmpty()== false))
+				{
+					TArray<FString> SearchTags;			
+					SearchTag->ParseIntoArray(SearchTags,TEXT(","));
+					//trim any xs spaces off the strings.
+					for (int32 iTag = 0; iTag < SearchTags.Num() ; iTag++)
+					{
+						SearchTags[iTag] = SearchTags[iTag].Trim();
+						SearchTags[iTag] = SearchTags[iTag].TrimTrailing();
+					}
+					if (SearchTags.Find(InText.ToString()) != INDEX_NONE)
+					{
+						FSearchEntry SearchEntry;
+						SearchEntry.Title = *ResultTitle;
+						SearchEntry.URL = "";
+						SearchEntry.bCategory = false;
+						SearchEntry.AssetData = Asset;
+						TutorialResults.Add(SearchEntry);
+					}
+				}
 			}
 		}
 #endif
diff --git a/Engine/Source/Editor/IntroTutorials/Public/EditorTutorial.h b/Engine/Source/Editor/IntroTutorials/Public/EditorTutorial.h
--- a/Engine/Source/Editor/IntroTutorials/Public/EditorTutorial.h
+++ b/Engine/Source/Editor/IntroTutorials/Public/EditorTutorial.h
@@ -278,6 +278,10 @@ class INTROTUTORIALS_API UEditorTutorial : public UObject
 	/** UObject implementation */
 	virtual UWorld* GetWorld() const override;
 
+	/** Comma seperated list of tags the search will use to help find this tutorial  */
+	UPROPERTY(EditAnywhere, Category = "Tutorial", AssetRegistrySearchable)
+	FText SearchTags;
+
 public:
 	/** Called when a tutorial stage is started */
 	void HandleTutorialStageStarted(FName StageName);
