Fixed mac compile error
(hopefully)

diff --git a/Engine/Source/Developer/SuperSearch/Public/SuperSearchModule.h b/Engine/Source/Developer/SuperSearch/Public/SuperSearchModule.h
--- a/Engine/Source/Developer/SuperSearch/Public/SuperSearchModule.h
+++ b/Engine/Source/Developer/SuperSearch/Public/SuperSearchModule.h
@@ -4,6 +4,10 @@
 
 #include "SlateBasics.h"
 
+#if WITH_EDITOR
+#include "AssetData.h"
+#endif
+
 struct FSearchEntry
 {
 	FString Title;
