[INTEGRATE] merge ue4 <-ue4-launcheDev cl 2633697

diff --git a/Engine/Source/Developer/SuperSearch/Private/SSuperSearch.cpp b/Engine/Source/Developer/SuperSearch/Private/SSuperSearch.cpp
--- a/Engine/Source/Developer/SuperSearch/Private/SSuperSearch.cpp
+++ b/Engine/Source/Developer/SuperSearch/Private/SSuperSearch.cpp
@@ -21,7 +21,7 @@ static TSharedRef<FSearchEntry> OtherCategory(new FSearchEntry());
 static TSharedRef<FSearchEntry> AskQuestionEntry (new FSearchEntry());
 
 SSuperSearchBox::SSuperSearchBox()
-	: SelectedSuggestion(-1)
+	: SelectedSuggestion(INDEX_NONE)
 	, bIgnoreUIUpdate(false)
 {
 	CategoryToIconMap.Add("Documentation", FName("LevelEditor.BrowseDocumentation") );
@@ -569,16 +569,15 @@ void SSuperSearchBox::UpdateSuggestions()
 	const FText & Query = InputText->GetText();
 	FSearchResults * SearchResults = SearchResultsCache.Find(Query.ToString());
 
+	//go through and build new suggestion list for list view widget
+	ClearSuggestions();
+
 	//still waiting on results for current query
 	if (SearchResults == NULL)
 	{
 		return;
 	}
 
-	//go through and build new suggestion list for list view widget
-	Suggestions.Empty();
-	SelectedSuggestion = -1;
-	
 	//first tutorials
 	UpdateSuggestionHelper(NSLOCTEXT("SuperSearch", "tutorials", "Tutorials"), SearchResults->TutorialResults, Suggestions);
 
@@ -633,7 +632,7 @@ void SSuperSearchBox::MarkActiveSuggestion()
 
 void SSuperSearchBox::ClearSuggestions()
 {
-	SelectedSuggestion = -1;
+	SelectedSuggestion = INDEX_NONE;
 	SuggestionBox->SetIsOpen(false);
 	Suggestions.Empty();
 }
