Reapplied CIS fix from CL# 3139679.

diff --git a/Engine/Source/Developer/SuperSearch/Private/SSuperSearch.cpp b/Engine/Source/Developer/SuperSearch/Private/SSuperSearch.cpp
--- a/Engine/Source/Developer/SuperSearch/Private/SSuperSearch.cpp
+++ b/Engine/Source/Developer/SuperSearch/Private/SSuperSearch.cpp
@@ -411,27 +411,27 @@ void SSuperSearchBox::OnTextChanged(const FText& InText)
 						SearchEntry.AssetData = Asset;
 						TutorialResults.Add(SearchEntry);
 					}
-				}
 
-				// If the asset has search tags, search them
-				if ((SearchTag) && (SearchTag->IsEmpty()== false))
-				{
-					TArray<FString> SearchTags;			
-					SearchTag->ParseIntoArray(SearchTags,TEXT(","));
-					//trim any xs spaces off the strings.
-					for (int32 iTag = 0; iTag < SearchTags.Num() ; iTag++)
-					{
-						SearchTags[iTag] = SearchTags[iTag].Trim();
-						SearchTags[iTag] = SearchTags[iTag].TrimTrailing();
-					}
-					if (SearchTags.Find(InText.ToString()) != INDEX_NONE)
+					// If the asset has search tags, search them
+					if ((SearchTag) && (SearchTag->IsEmpty()== false))
 					{
-						FSearchEntry SearchEntry;
-						SearchEntry.Title = *ResultTitle;
-						SearchEntry.URL = "";
-						SearchEntry.bCategory = false;
-						SearchEntry.AssetData = Asset;
-						TutorialResults.Add(SearchEntry);
+						TArray<FString> SearchTags;			
+						SearchTag->ParseIntoArray(SearchTags,TEXT(","));
+						//trim any xs spaces off the strings.
+						for (int32 iTag = 0; iTag < SearchTags.Num() ; iTag++)
+						{
+							SearchTags[iTag] = SearchTags[iTag].Trim();
+							SearchTags[iTag] = SearchTags[iTag].TrimTrailing();
+						}
+						if (SearchTags.Find(InText.ToString()) != INDEX_NONE)
+						{
+							FSearchEntry SearchEntry;
+							SearchEntry.Title = *ResultTitle;
+							SearchEntry.URL = "";
+							SearchEntry.bCategory = false;
+							SearchEntry.AssetData = Asset;
+							TutorialResults.Add(SearchEntry);
+						}
 					}
 				}
 			}
